Fixed missing copyrights (Horde issue 195346)

diff --git a/Engine/Plugins/Enterprise/GLTFExporter/Source/GLTFExporter/Private/UI/GLTFExportOptionsWindow.cpp b/Engine/Plugins/Enterprise/GLTFExporter/Source/GLTFExporter/Private/UI/GLTFExportOptionsWindow.cpp
--- a/Engine/Plugins/Enterprise/GLTFExporter/Source/GLTFExporter/Private/UI/GLTFExportOptionsWindow.cpp
+++ b/Engine/Plugins/Enterprise/GLTFExporter/Source/GLTFExporter/Private/UI/GLTFExportOptionsWindow.cpp
@@ -1,4 +1,4 @@
-// Copyright 1998-2019 Epic Games, Inc. All Rights Reserved.
+// Copyright Epic Games, Inc. All Rights Reserved.
 
 #if WITH_EDITOR
 
diff --git a/Engine/Plugins/Enterprise/GLTFExporter/Source/GLTFExporter/Private/UI/GLTFExportOptionsWindow.h b/Engine/Plugins/Enterprise/GLTFExporter/Source/GLTFExporter/Private/UI/GLTFExportOptionsWindow.h
--- a/Engine/Plugins/Enterprise/GLTFExporter/Source/GLTFExporter/Private/UI/GLTFExportOptionsWindow.h
+++ b/Engine/Plugins/Enterprise/GLTFExporter/Source/GLTFExporter/Private/UI/GLTFExportOptionsWindow.h
@@ -1,4 +1,4 @@
-// Copyright 1998-2019 Epic Games, Inc. All Rights Reserved.
+// Copyright Epic Games, Inc. All Rights Reserved.
 
 #pragma once
 
