Fixed static analysis issue

diff --git a/Engine/Plugins/Enterprise/GLTFExporter/Source/GLTFExporter/Private/Converters/GLTFUVOverlapChecker.cpp b/Engine/Plugins/Enterprise/GLTFExporter/Source/GLTFExporter/Private/Converters/GLTFUVOverlapChecker.cpp
--- a/Engine/Plugins/Enterprise/GLTFExporter/Source/GLTFExporter/Private/Converters/GLTFUVOverlapChecker.cpp
+++ b/Engine/Plugins/Enterprise/GLTFExporter/Source/GLTFExporter/Private/Converters/GLTFUVOverlapChecker.cpp
@@ -14,7 +14,8 @@
 
 void FGLTFUVOverlapChecker::Sanitize(const FMeshDescription*& Description, FGLTFIndexArray& SectionIndices, int32& TexCoord)
 {
-#if WITH_EDITOR
+	// The code below is disabled waiting for the proper fix: See UE-159948
+#if WITH_EDITOR && !WITH_EDITOR
 	if (Description != nullptr)
 	{
 		const TVertexInstanceAttributesConstRef<FVector2f> VertexInstanceUVs =
