Fixed static analysis issue found by Horde

diff --git a/Engine/Plugins/Enterprise/GLTFExporter/Source/GLTFExporter/Private/Tasks/GLTFDelayedMaterialTasks.cpp b/Engine/Plugins/Enterprise/GLTFExporter/Source/GLTFExporter/Private/Tasks/GLTFDelayedMaterialTasks.cpp
--- a/Engine/Plugins/Enterprise/GLTFExporter/Source/GLTFExporter/Private/Tasks/GLTFDelayedMaterialTasks.cpp
+++ b/Engine/Plugins/Enterprise/GLTFExporter/Source/GLTFExporter/Private/Tasks/GLTFDelayedMaterialTasks.cpp
@@ -684,7 +684,7 @@ bool FGLTFDelayedMaterialTask::TryGetMetallicAndRoughness(FGLTFJsonPBRMetallicRo
 	}
 
 	int32 TexCoord;
-	if (MetallicBakeOutput.bIsConstant)
+	if (RoughnessBakeOutput.bIsConstant)
 	{
 		TexCoord = RoughnessTexCoord;
 	}
@@ -809,7 +809,7 @@ bool FGLTFDelayedMaterialTask::TryGetClearCoatRoughness(FGLTFJsonClearCoatExtens
 	}
 
 	int32 TexCoord;
-	if (IntensityBakeOutput.bIsConstant)
+	if (RoughnessBakeOutput.bIsConstant)
 	{
 		TexCoord = RoughnessTexCoord;
 	}
