Resolve comment #ueent_correction_to_review in GLTFNodeConverters.cpp

diff --git a/Engine/Plugins/Enterprise/GLTFExporter/Source/GLTFExporter/Private/Converters/GLTFNodeConverters.cpp b/Engine/Plugins/Enterprise/GLTFExporter/Source/GLTFExporter/Private/Converters/GLTFNodeConverters.cpp
--- a/Engine/Plugins/Enterprise/GLTFExporter/Source/GLTFExporter/Private/Converters/GLTFNodeConverters.cpp
+++ b/Engine/Plugins/Enterprise/GLTFExporter/Source/GLTFExporter/Private/Converters/GLTFNodeConverters.cpp
@@ -220,14 +220,11 @@ FGLTFJsonNode* FGLTFComponentSocketConverter::Convert(const USceneComponent* Sce
 			return Builder.AddUniqueNode(JsonNode, StaticMesh, SocketName);
 		}
 
-		if (const USkinnedMeshComponent* SkinnedMeshComponent = Cast<USkinnedMeshComponent>(SceneComponent))
+		if (const USkeletalMeshComponent* SkeletalMeshComponent = Cast<USkeletalMeshComponent>(SceneComponent))
 		{
 			// TODO: add support for SocketOverrideLookup?
-			// #ueent_correction_to_review
-			if (const USkeletalMesh* SkeletalMesh = Cast<USkeletalMesh>(SkinnedMeshComponent->GetSkinnedAsset()))
-			{
-				return Builder.AddUniqueNode(JsonNode, SkeletalMesh, SocketName);
-			}
+			const USkeletalMesh* SkeletalMesh = SkeletalMeshComponent->GetSkeletalMeshAsset();
+			return Builder.AddUniqueNode(JsonNode, SkeletalMesh, SocketName);
 		}
 
 		// TODO: add support for more socket types
