Fix FGLTFProxyAssetActions::GetClassPathName, so it properly forwards the original�s value

diff --git a/Engine/Plugins/Enterprise/GLTFExporter/Source/GLTFExporter/Private/Actions/GLTFProxyAssetActions.h b/Engine/Plugins/Enterprise/GLTFExporter/Source/GLTFExporter/Private/Actions/GLTFProxyAssetActions.h
--- a/Engine/Plugins/Enterprise/GLTFExporter/Source/GLTFExporter/Private/Actions/GLTFProxyAssetActions.h
+++ b/Engine/Plugins/Enterprise/GLTFExporter/Source/GLTFExporter/Private/Actions/GLTFProxyAssetActions.h
@@ -46,7 +46,7 @@ public:
 	virtual FText GetDisplayNameFromAssetData(const FAssetData& AssetData) const override { return OriginalActions->GetDisplayNameFromAssetData(AssetData); }
 	virtual void SetSupported(bool bInSupported) override { OriginalActions->SetSupported(bInSupported); }
 	virtual bool IsSupported() const override { return OriginalActions->IsSupported(); }
-	virtual FTopLevelAssetPath GetClassPathName() const override { return TopLevelAssetPath; };
+	virtual FTopLevelAssetPath GetClassPathName() const override { return OriginalActions->GetClassPathName(); };
 
 private:
 
@@ -55,8 +55,6 @@ private:
 	void OnCreateProxy(TArray<FWeakObjectPtr> Objects) const;
 
 	TSharedRef<IAssetTypeActions> OriginalActions;
-
-	FTopLevelAssetPath TopLevelAssetPath;
 };
 
 #endif
