Fix selection of texcoord (for combined gltf material properties) and refactor fixed code for improved clarity

diff --git a/Engine/Plugins/Enterprise/GLTFExporter/Source/GLTFExporter/Private/Tasks/GLTFDelayedMaterialTasks.cpp b/Engine/Plugins/Enterprise/GLTFExporter/Source/GLTFExporter/Private/Tasks/GLTFDelayedMaterialTasks.cpp
--- a/Engine/Plugins/Enterprise/GLTFExporter/Source/GLTFExporter/Private/Tasks/GLTFDelayedMaterialTasks.cpp
+++ b/Engine/Plugins/Enterprise/GLTFExporter/Source/GLTFExporter/Private/Tasks/GLTFDelayedMaterialTasks.cpp
@@ -559,18 +559,16 @@ bool FGLTFDelayedMaterialTask::TryGetBaseColorAndOpacity(FGLTFJsonPBRMetallicRou
 		return true;
 	}
 
-	int32 TexCoord;
-	if (BaseColorBakeOutput.bIsConstant)
+	int32 CombinedTexCoord;
+
+	// If one is constant, it means we can use the other's texture coords
+	if (BaseColorBakeOutput.bIsConstant || BaseColorTexCoord == OpacityTexCoord)
 	{
-		TexCoord = OpacityTexCoord;
+		CombinedTexCoord = OpacityTexCoord;
 	}
 	else if (OpacityBakeOutput.bIsConstant)
 	{
-		TexCoord = BaseColorTexCoord;
-	}
-	else if (BaseColorTexCoord == OpacityTexCoord)
-	{
-		TexCoord = BaseColorTexCoord;
+		CombinedTexCoord = BaseColorTexCoord;
 	}
 	else
 	{
@@ -588,7 +586,7 @@ bool FGLTFDelayedMaterialTask::TryGetBaseColorAndOpacity(FGLTFJsonPBRMetallicRou
 			return FColor(BaseColor.R, BaseColor.G, BaseColor.B, Opacity.R);
 		});
 
-	FGLTFJsonTexture* Texture = FGLTFMaterialUtility::AddTexture(
+	FGLTFJsonTexture* CombinedTexture = FGLTFMaterialUtility::AddTexture(
 		Builder,
 		CombinedPixels,
 		TextureSize,
@@ -600,8 +598,8 @@ bool FGLTFDelayedMaterialTask::TryGetBaseColorAndOpacity(FGLTFJsonPBRMetallicRou
 		TextureWrapS,
 		TextureWrapT);
 
-	OutPBRParams.BaseColorTexture.TexCoord = TexCoord;
-	OutPBRParams.BaseColorTexture.Index = Texture;
+	OutPBRParams.BaseColorTexture.TexCoord = CombinedTexCoord;
+	OutPBRParams.BaseColorTexture.Index = CombinedTexture;
 	OutPBRParams.BaseColorFactor = FGLTFCoreUtilities::ConvertColor({ BaseColorScale, BaseColorScale, BaseColorScale }, Builder.ExportOptions->bStrictCompliance);
 
 	return true;
@@ -683,18 +681,16 @@ bool FGLTFDelayedMaterialTask::TryGetMetallicAndRoughness(FGLTFJsonPBRMetallicRo
 		return true;
 	}
 
-	int32 TexCoord;
-	if (RoughnessBakeOutput.bIsConstant)
+	int32 CombinedTexCoord;
+	
+	// If one is constant, it means we can use the other's texture coords
+	if (MetallicBakeOutput.bIsConstant || MetallicTexCoord == RoughnessTexCoord)
 	{
-		TexCoord = RoughnessTexCoord;
+		CombinedTexCoord = RoughnessTexCoord;
 	}
-	else if (MetallicBakeOutput.bIsConstant)
+	else if (RoughnessBakeOutput.bIsConstant)
 	{
-		TexCoord = MetallicTexCoord;
-	}
-	else if (MetallicTexCoord == RoughnessTexCoord)
-	{
-		TexCoord = MetallicTexCoord;
+		CombinedTexCoord = MetallicTexCoord;
 	}
 	else
 	{
@@ -712,7 +708,7 @@ bool FGLTFDelayedMaterialTask::TryGetMetallicAndRoughness(FGLTFJsonPBRMetallicRo
 			return FColor(0, Roughness.R, Metallic.R);
 		});
 
-	FGLTFJsonTexture* Texture = FGLTFMaterialUtility::AddTexture(
+	FGLTFJsonTexture* CombinedTexture = FGLTFMaterialUtility::AddTexture(
 		Builder,
 		CombinedPixels,
 		TextureSize,
@@ -724,8 +720,8 @@ bool FGLTFDelayedMaterialTask::TryGetMetallicAndRoughness(FGLTFJsonPBRMetallicRo
 		TextureWrapS,
 		TextureWrapT);
 
-	OutPBRParams.MetallicRoughnessTexture.TexCoord = TexCoord;
-	OutPBRParams.MetallicRoughnessTexture.Index = Texture;
+	OutPBRParams.MetallicRoughnessTexture.TexCoord = CombinedTexCoord;
+	OutPBRParams.MetallicRoughnessTexture.Index = CombinedTexture;
 
 	return true;
 }
@@ -808,18 +804,16 @@ bool FGLTFDelayedMaterialTask::TryGetClearCoatRoughness(FGLTFJsonClearCoatExtens
 		return true;
 	}
 
-	int32 TexCoord;
-	if (RoughnessBakeOutput.bIsConstant)
+	int32 CombinedTexCoord;
+	
+	// If one is constant, it means we can use the other's texture coords
+	if (IntensityBakeOutput.bIsConstant || IntensityTexCoord == RoughnessTexCoord)
 	{
-		TexCoord = RoughnessTexCoord;
+		CombinedTexCoord = RoughnessTexCoord;
 	}
-	else if (IntensityBakeOutput.bIsConstant)
+	else if (RoughnessBakeOutput.bIsConstant)
 	{
-		TexCoord = IntensityTexCoord;
-	}
-	else if (IntensityTexCoord == RoughnessTexCoord)
-	{
-		TexCoord = IntensityTexCoord;
+		CombinedTexCoord = IntensityTexCoord;
 	}
 	else
 	{
@@ -837,7 +831,7 @@ bool FGLTFDelayedMaterialTask::TryGetClearCoatRoughness(FGLTFJsonClearCoatExtens
 			return FColor(Intensity.R, Roughness.R, 0);
 		});
 
-	FGLTFJsonTexture* Texture = FGLTFMaterialUtility::AddTexture(
+	FGLTFJsonTexture* CombinedTexture = FGLTFMaterialUtility::AddTexture(
 		Builder,
 		CombinedPixels,
 		TextureSize,
@@ -849,10 +843,10 @@ bool FGLTFDelayedMaterialTask::TryGetClearCoatRoughness(FGLTFJsonClearCoatExtens
 		TextureWrapS,
 		TextureWrapT);
 
-	OutExtParams.ClearCoatTexture.Index = Texture;
-	OutExtParams.ClearCoatTexture.TexCoord = TexCoord;
-	OutExtParams.ClearCoatRoughnessTexture.Index = Texture;
-	OutExtParams.ClearCoatRoughnessTexture.TexCoord = TexCoord;
+	OutExtParams.ClearCoatTexture.Index = CombinedTexture;
+	OutExtParams.ClearCoatTexture.TexCoord = CombinedTexCoord;
+	OutExtParams.ClearCoatRoughnessTexture.Index = CombinedTexture;
+	OutExtParams.ClearCoatRoughnessTexture.TexCoord = CombinedTexCoord;
 
 	return true;
 }
