Revisited temporary fix and removed unnecessary checks that silences bad parameters

diff --git a/Engine/Plugins/Enterprise/GLTFExporter/Source/GLTFExporter/Private/Converters/GLTFUVDegenerateChecker.cpp b/Engine/Plugins/Enterprise/GLTFExporter/Source/GLTFExporter/Private/Converters/GLTFUVDegenerateChecker.cpp
--- a/Engine/Plugins/Enterprise/GLTFExporter/Source/GLTFExporter/Private/Converters/GLTFUVDegenerateChecker.cpp
+++ b/Engine/Plugins/Enterprise/GLTFExporter/Source/GLTFExporter/Private/Converters/GLTFUVDegenerateChecker.cpp
@@ -6,32 +6,6 @@
 #include "MeshDescription.h"
 #endif
 
-void FGLTFUVDegenerateChecker::Sanitize(const FMeshDescription*& Description, FGLTFIndexArray& SectionIndices, int32& TexCoord)
-{
-#if WITH_EDITOR
-	if (Description != nullptr)
-	{
-		const int32 SectionCount = Description->PolygonGroups().Num();
-		for (const int32 SectionIndex : SectionIndices)
-		{
-			if (SectionIndex < 0 || SectionIndex >= SectionCount)
-			{
-				Description = nullptr;
-			}
-		}
-	}
-
-	if (Description != nullptr)
-	{
-		const int32 TexCoordCount = Description->VertexInstanceAttributes().GetAttributesRef<FVector2f>(MeshAttribute::VertexInstance::TextureCoordinate).GetNumChannels();
-		if (TexCoord < 0 || TexCoord >= TexCoordCount)
-		{
-			Description = nullptr;
-		}
-	}
-#endif
-}
-
 float FGLTFUVDegenerateChecker::Convert(const FMeshDescription* Description, FGLTFIndexArray SectionIndices, int32 TexCoord)
 {
 #if WITH_EDITOR
diff --git a/Engine/Plugins/Enterprise/GLTFExporter/Source/GLTFExporter/Private/Converters/GLTFUVOverlapChecker.cpp b/Engine/Plugins/Enterprise/GLTFExporter/Source/GLTFExporter/Private/Converters/GLTFUVOverlapChecker.cpp
--- a/Engine/Plugins/Enterprise/GLTFExporter/Source/GLTFExporter/Private/Converters/GLTFUVOverlapChecker.cpp
+++ b/Engine/Plugins/Enterprise/GLTFExporter/Source/GLTFExporter/Private/Converters/GLTFUVOverlapChecker.cpp
@@ -12,33 +12,6 @@
 #include "UObject/UObjectGlobals.h"
 #endif
 
-void FGLTFUVOverlapChecker::Sanitize(const FMeshDescription*& Description, FGLTFIndexArray& SectionIndices, int32& TexCoord)
-{
-	// The code below is disabled waiting for the proper fix: See UE-159948
-#if WITH_EDITOR && !WITH_EDITOR
-	if (Description != nullptr)
-	{
-		const TVertexInstanceAttributesConstRef<FVector2f> VertexInstanceUVs =
-			Description->VertexInstanceAttributes().GetAttributesRef<FVector2f>(MeshAttribute::VertexInstance::TextureCoordinate);
-		const int32 TexCoordCount = VertexInstanceUVs.GetNumChannels();
-
-		if (TexCoord < 0 || TexCoord >= TexCoordCount)
-		{
-			Description = nullptr;
-		}
-
-		const int32 MinSectionIndex = FMath::Min(SectionIndices);
-		const int32 MaxSectionIndex = FMath::Max(SectionIndices);
-		const int32 SectionCount = Description->PolygonGroups().GetArraySize();
-
-		if (MinSectionIndex < 0 || MaxSectionIndex >= SectionCount)
-		{
-			Description = nullptr;
-		}
-	}
-#endif
-}
-
 float FGLTFUVOverlapChecker::Convert(const FMeshDescription* Description, FGLTFIndexArray SectionIndices, int32 TexCoord)
 {
 #if WITH_EDITOR
diff --git a/Engine/Plugins/Enterprise/GLTFExporter/Source/GLTFExporter/Public/Converters/GLTFUVDegenerateChecker.h b/Engine/Plugins/Enterprise/GLTFExporter/Source/GLTFExporter/Public/Converters/GLTFUVDegenerateChecker.h
--- a/Engine/Plugins/Enterprise/GLTFExporter/Source/GLTFExporter/Public/Converters/GLTFUVDegenerateChecker.h
+++ b/Engine/Plugins/Enterprise/GLTFExporter/Source/GLTFExporter/Public/Converters/GLTFUVDegenerateChecker.h
@@ -14,8 +14,6 @@ class GLTFEXPORTER_API FGLTFUVDegenerateChecker : public TGLTFConverter<float, c
 {
 protected:
 
-	virtual void Sanitize(const FMeshDescription*& Description, FGLTFIndexArray& SectionIndices, int32& TexCoord) override;
-
 	virtual float Convert(const FMeshDescription* Description, FGLTFIndexArray SectionIndices, int32 TexCoord) override;
 
 private:
diff --git a/Engine/Plugins/Enterprise/GLTFExporter/Source/GLTFExporter/Public/Converters/GLTFUVOverlapChecker.h b/Engine/Plugins/Enterprise/GLTFExporter/Source/GLTFExporter/Public/Converters/GLTFUVOverlapChecker.h
--- a/Engine/Plugins/Enterprise/GLTFExporter/Source/GLTFExporter/Public/Converters/GLTFUVOverlapChecker.h
+++ b/Engine/Plugins/Enterprise/GLTFExporter/Source/GLTFExporter/Public/Converters/GLTFUVOverlapChecker.h
@@ -12,8 +12,6 @@ class GLTFEXPORTER_API FGLTFUVOverlapChecker : public TGLTFConverter<float, cons
 {
 protected:
 
-	virtual void Sanitize(const FMeshDescription*& Description, FGLTFIndexArray& SectionIndices, int32& TexCoord) override;
-
 	virtual float Convert(const FMeshDescription* Description, FGLTFIndexArray SectionIndices, int32 TexCoord) override;
 
 private:
