Remove GLTFExporter's backwards compatibility for previous engines (since the plugin is distributed with specific engine now)

diff --git a/Engine/Plugins/Enterprise/GLTFExporter/Source/GLTFExporter/Private/Actors/GLTFHotspotActor.cpp b/Engine/Plugins/Enterprise/GLTFExporter/Source/GLTFExporter/Private/Actors/GLTFHotspotActor.cpp
--- a/Engine/Plugins/Enterprise/GLTFExporter/Source/GLTFExporter/Private/Actors/GLTFHotspotActor.cpp
+++ b/Engine/Plugins/Enterprise/GLTFExporter/Source/GLTFExporter/Private/Actors/GLTFHotspotActor.cpp
@@ -20,9 +20,7 @@
 #include "Slate/SceneViewport.h"
 #endif
 
-#if (ENGINE_MAJOR_VERSION > 4 || ENGINE_MINOR_VERSION >= 26)
 using namespace UE;
-#endif
 
 DEFINE_LOG_CATEGORY_STATIC(LogGLTFHotspot, Log, All);
 
@@ -455,11 +453,7 @@ void AGLTFHotspotActor::ValidateAnimation()
 	{
 		if (const USkeletalMesh* SkeletalMesh = SkeletalMeshActor->GetSkeletalMeshComponent()->GetSkeletalMeshAsset())
 		{
-#if (ENGINE_MAJOR_VERSION > 4 || ENGINE_MINOR_VERSION >= 27)
 			const USkeleton* Skeleton = SkeletalMesh->GetSkeleton();
-#else
-			const USkeleton* Skeleton = SkeletalMesh->Skeleton;
-#endif
 
 			if (Skeleton != AnimationSequence->GetSkeleton())
 			{
diff --git a/Engine/Plugins/Enterprise/GLTFExporter/Source/GLTFExporter/Private/Converters/GLTFAnimationConverters.cpp b/Engine/Plugins/Enterprise/GLTFExporter/Source/GLTFExporter/Private/Converters/GLTFAnimationConverters.cpp
--- a/Engine/Plugins/Enterprise/GLTFExporter/Source/GLTFExporter/Private/Converters/GLTFAnimationConverters.cpp
+++ b/Engine/Plugins/Enterprise/GLTFExporter/Source/GLTFExporter/Private/Converters/GLTFAnimationConverters.cpp
@@ -21,11 +21,7 @@ FGLTFJsonAnimation* FGLTFAnimationConverter::Convert(FGLTFJsonNode* RootNode, co
 		return nullptr;
 	}
 
-#if (ENGINE_MAJOR_VERSION > 4 || ENGINE_MINOR_VERSION >= 27)
 	const USkeleton* MeshSkeleton = SkeletalMesh->GetSkeleton();
-#else
-	const USkeleton* MeshSkeleton = SkeletalMesh->Skeleton;
-#endif
 
 	if (AnimSkeleton != MeshSkeleton)
 	{
diff --git a/Engine/Plugins/Enterprise/GLTFExporter/Source/GLTFExporter/Private/Converters/GLTFBufferUtility.cpp b/Engine/Plugins/Enterprise/GLTFExporter/Source/GLTFExporter/Private/Converters/GLTFBufferUtility.cpp
--- a/Engine/Plugins/Enterprise/GLTFExporter/Source/GLTFExporter/Private/Converters/GLTFBufferUtility.cpp
+++ b/Engine/Plugins/Enterprise/GLTFExporter/Source/GLTFExporter/Private/Converters/GLTFBufferUtility.cpp
@@ -102,25 +102,7 @@ const void* FGLTFBufferUtility::GetCPUBuffer(const FRawStaticIndexBuffer16or32In
 
 const void* FGLTFBufferUtility::GetCPUBuffer(const FSkinWeightDataVertexBuffer* VertexBuffer)
 {
-#if (ENGINE_MAJOR_VERSION > 4 || ENGINE_MINOR_VERSION >= 26)
 	return VertexBuffer->GetWeightData();
-#else
-	struct FSkinWeightDataVertexBufferHack : FVertexBuffer
-	{
-		FShaderResourceViewRHIRef SRVValue;
-		bool bNeedsCPUAccess;
-		bool bVariableBonesPerVertex;
-		uint32 MaxBoneInfluences;
-		bool bUse16BitBoneIndex;
-		FStaticMeshVertexDataInterface* WeightData;
-		uint8* Data;
-		uint32 NumVertices;
-		uint32 NumBones;
-	};
-
-	static_assert(sizeof(FSkinWeightDataVertexBufferHack) == sizeof(FSkinWeightDataVertexBuffer), "FSkinWeightDataVertexBufferHack memory layout doesn't match FSkinWeightDataVertexBuffer");
-	return reinterpret_cast<const FSkinWeightDataVertexBufferHack*>(VertexBuffer)->Data;
-#endif
 }
 
 const void* FGLTFBufferUtility::GetCPUBuffer(const FSkinWeightLookupVertexBuffer* VertexBuffer)
diff --git a/Engine/Plugins/Enterprise/GLTFExporter/Source/GLTFExporter/Private/Converters/GLTFLightMapConverters.cpp b/Engine/Plugins/Enterprise/GLTFExporter/Source/GLTFExporter/Private/Converters/GLTFLightMapConverters.cpp
--- a/Engine/Plugins/Enterprise/GLTFExporter/Source/GLTFExporter/Private/Converters/GLTFLightMapConverters.cpp
+++ b/Engine/Plugins/Enterprise/GLTFExporter/Source/GLTFExporter/Private/Converters/GLTFLightMapConverters.cpp
@@ -23,12 +23,7 @@ FGLTFJsonLightMap* FGLTFLightMapConverter::Convert(const UStaticMeshComponent* S
 
 	const int32 LODIndex = FGLTFMeshUtility::GetLOD(StaticMesh, StaticMeshComponent, Builder.ExportOptions->DefaultLevelOfDetail);
 	const FStaticMeshLODResources& LODResources = StaticMesh->GetLODForExport(LODIndex);
-
-#if (ENGINE_MAJOR_VERSION > 4 || ENGINE_MINOR_VERSION >= 27)
 	const int32 CoordinateIndex = StaticMesh->GetLightMapCoordinateIndex();
-#else
-	const int32 CoordinateIndex = StaticMesh->LightMapCoordinateIndex;
-#endif
 
 	if (CoordinateIndex < 0 || CoordinateIndex >= LODResources.GetNumTexCoords())
 	{
diff --git a/Engine/Plugins/Enterprise/GLTFExporter/Source/GLTFExporter/Private/Converters/GLTFMaterialUtility.cpp b/Engine/Plugins/Enterprise/GLTFExporter/Source/GLTFExporter/Private/Converters/GLTFMaterialUtility.cpp
--- a/Engine/Plugins/Enterprise/GLTFExporter/Source/GLTFExporter/Private/Converters/GLTFMaterialUtility.cpp
+++ b/Engine/Plugins/Enterprise/GLTFExporter/Source/GLTFExporter/Private/Converters/GLTFMaterialUtility.cpp
@@ -34,7 +34,6 @@ bool FGLTFMaterialUtility::IsSRGB(const FGLTFMaterialPropertyEx& Property)
 	return Property == MP_BaseColor || Property == MP_EmissiveColor || Property == MP_SubsurfaceColor || Property == TEXT("TransmittanceColor");
 }
 
-#if (ENGINE_MAJOR_VERSION > 4 || ENGINE_MINOR_VERSION >= 27)
 FGuid FGLTFMaterialUtility::GetAttributeID(const FGLTFMaterialPropertyEx& Property)
 {
 	return Property.IsCustomOutput()
@@ -48,7 +47,6 @@ FGuid FGLTFMaterialUtility::GetAttributeIDChecked(const FGLTFMaterialPropertyEx&
 	check(AttributeID != FMaterialAttributeDefinitionMap::GetDefaultID());
 	return AttributeID;
 }
-#endif
 
 FVector4f FGLTFMaterialUtility::GetPropertyDefaultValue(const FGLTFMaterialPropertyEx& Property)
 {
diff --git a/Engine/Plugins/Enterprise/GLTFExporter/Source/GLTFExporter/Private/Converters/GLTFMaterialUtility.h b/Engine/Plugins/Enterprise/GLTFExporter/Source/GLTFExporter/Private/Converters/GLTFMaterialUtility.h
--- a/Engine/Plugins/Enterprise/GLTFExporter/Source/GLTFExporter/Private/Converters/GLTFMaterialUtility.h
+++ b/Engine/Plugins/Enterprise/GLTFExporter/Source/GLTFExporter/Private/Converters/GLTFMaterialUtility.h
@@ -36,10 +36,8 @@ struct FGLTFMaterialUtility
 	static bool IsNormalMap(const FGLTFMaterialPropertyEx& Property);
 	static bool IsSRGB(const FGLTFMaterialPropertyEx& Property);
 
-#if (ENGINE_MAJOR_VERSION > 4 || ENGINE_MINOR_VERSION >= 27)
 	static FGuid GetAttributeID(const FGLTFMaterialPropertyEx& Property);
 	static FGuid GetAttributeIDChecked(const FGLTFMaterialPropertyEx& Property);
-#endif
 
 	static FVector4f GetPropertyDefaultValue(const FGLTFMaterialPropertyEx& Property);
 	static FVector4f GetPropertyMask(const FGLTFMaterialPropertyEx& Property);
diff --git a/Engine/Plugins/Enterprise/GLTFExporter/Source/GLTFExporter/Private/Converters/GLTFMeshData.cpp b/Engine/Plugins/Enterprise/GLTFExporter/Source/GLTFExporter/Private/Converters/GLTFMeshData.cpp
--- a/Engine/Plugins/Enterprise/GLTFExporter/Source/GLTFExporter/Private/Converters/GLTFMeshData.cpp
+++ b/Engine/Plugins/Enterprise/GLTFExporter/Source/GLTFExporter/Private/Converters/GLTFMeshData.cpp
@@ -60,13 +60,7 @@ FGLTFMeshData::FGLTFMeshData(const UStaticMesh* StaticMesh, const UStaticMeshCom
 	}
 
 #if WITH_EDITOR
-
-#if (ENGINE_MAJOR_VERSION > 4 || ENGINE_MINOR_VERSION >= 27)
 	LightMapTexCoord = StaticMesh->GetLightMapCoordinateIndex();
-#else
-	LightMapTexCoord = StaticMesh->LightMapCoordinateIndex;
-#endif
-
  	const int32 NumTexCoords = StaticMesh->GetLODForExport(LODIndex).VertexBuffers.StaticMeshVertexBuffer.GetNumTexCoords();
 	BakeUsingTexCoord = FMath::Min(LightMapTexCoord, NumTexCoords - 1);
 	// TODO: add warning if texture coordinate for baking has overlap
diff --git a/Engine/Plugins/Enterprise/GLTFExporter/Source/GLTFExporter/Private/Converters/GLTFMeshUtility.cpp b/Engine/Plugins/Enterprise/GLTFExporter/Source/GLTFExporter/Private/Converters/GLTFMeshUtility.cpp
--- a/Engine/Plugins/Enterprise/GLTFExporter/Source/GLTFExporter/Private/Converters/GLTFMeshUtility.cpp
+++ b/Engine/Plugins/Enterprise/GLTFExporter/Source/GLTFExporter/Private/Converters/GLTFMeshUtility.cpp
@@ -41,20 +41,12 @@ const UMaterialInterface* FGLTFMeshUtility::GetDefaultMaterial()
 
 const TArray<FStaticMaterial>& FGLTFMeshUtility::GetMaterials(const UStaticMesh* StaticMesh)
 {
-#if (ENGINE_MAJOR_VERSION > 4 || ENGINE_MINOR_VERSION >= 27)
 	return StaticMesh->GetStaticMaterials();
-#else
-	return StaticMesh->StaticMaterials;
-#endif
 }
 
 const TArray<FSkeletalMaterial>& FGLTFMeshUtility::GetMaterials(const USkeletalMesh* SkeletalMesh)
 {
-#if (ENGINE_MAJOR_VERSION > 4 || ENGINE_MINOR_VERSION >= 27)
 	return SkeletalMesh->GetMaterials();
-#else
-	return SkeletalMesh->Materials;
-#endif
 }
 
 const UMaterialInterface* FGLTFMeshUtility::GetMaterial(const UMaterialInterface* Material)
@@ -230,11 +222,7 @@ int32 FGLTFMeshUtility::GetMinimumLOD(const UStaticMesh* StaticMesh, const UStat
 
 	if (StaticMesh != nullptr)
 	{
-#if (ENGINE_MAJOR_VERSION > 4 || ENGINE_MINOR_VERSION >= 27)
 		return StaticMesh->GetMinLOD().Default;
-#else
-		return StaticMesh->MinLOD.Default;
-#endif
 	}
 
 	return -1;
@@ -249,11 +237,7 @@ int32 FGLTFMeshUtility::GetMinimumLOD(const USkeletalMesh* SkeletalMesh, const U
 
 	if (SkeletalMesh != nullptr)
 	{
-#if (ENGINE_MAJOR_VERSION > 4 || ENGINE_MINOR_VERSION >= 27)
 		return SkeletalMesh->GetMinLod().Default;
-#else
-		return SkeletalMesh->MinLod.Default;
-#endif
 	}
 
 	return -1;
diff --git a/Engine/Plugins/Enterprise/GLTFExporter/Source/GLTFExporter/Private/Converters/GLTFNodeConverters.cpp b/Engine/Plugins/Enterprise/GLTFExporter/Source/GLTFExporter/Private/Converters/GLTFNodeConverters.cpp
--- a/Engine/Plugins/Enterprise/GLTFExporter/Source/GLTFExporter/Private/Converters/GLTFNodeConverters.cpp
+++ b/Engine/Plugins/Enterprise/GLTFExporter/Source/GLTFExporter/Private/Converters/GLTFNodeConverters.cpp
@@ -261,11 +261,7 @@ FGLTFJsonNode* FGLTFStaticSocketConverter::Convert(FGLTFJsonNode* RootNode, cons
 
 FGLTFJsonNode* FGLTFSkeletalSocketConverter::Convert(FGLTFJsonNode* RootNode, const USkeletalMesh* SkeletalMesh, FName SocketName)
 {
-#if (ENGINE_MAJOR_VERSION > 4 || ENGINE_MINOR_VERSION >= 27)
 	const FReferenceSkeleton& RefSkeleton = SkeletalMesh->GetRefSkeleton();
-#else
-	const FReferenceSkeleton& RefSkeleton = SkeletalMesh->RefSkeleton;
-#endif
 
 	const USkeletalMeshSocket* Socket = SkeletalMesh->FindSocket(SocketName);
 	if (Socket != nullptr)
@@ -297,11 +293,7 @@ FGLTFJsonNode* FGLTFSkeletalSocketConverter::Convert(FGLTFJsonNode* RootNode, co
 
 FGLTFJsonNode* FGLTFSkeletalBoneConverter::Convert(FGLTFJsonNode* RootNode, const USkeletalMesh* SkeletalMesh, int32 BoneIndex)
 {
-#if (ENGINE_MAJOR_VERSION > 4 || ENGINE_MINOR_VERSION >= 27)
 	const FReferenceSkeleton& RefSkeleton = SkeletalMesh->GetRefSkeleton();
-#else
-	const FReferenceSkeleton& RefSkeleton = SkeletalMesh->RefSkeleton;
-#endif
 
 	// TODO: add support for [Principal]PoseComponent?
 
diff --git a/Engine/Plugins/Enterprise/GLTFExporter/Source/GLTFExporter/Private/Converters/GLTFSkinConverters.cpp b/Engine/Plugins/Enterprise/GLTFExporter/Source/GLTFExporter/Private/Converters/GLTFSkinConverters.cpp
--- a/Engine/Plugins/Enterprise/GLTFExporter/Source/GLTFExporter/Private/Converters/GLTFSkinConverters.cpp
+++ b/Engine/Plugins/Enterprise/GLTFExporter/Source/GLTFExporter/Private/Converters/GLTFSkinConverters.cpp
@@ -11,13 +11,8 @@
 
 FGLTFJsonSkin* FGLTFSkinConverter::Convert(FGLTFJsonNode* RootNode, const USkeletalMesh* SkeletalMesh)
 {
-#if (ENGINE_MAJOR_VERSION > 4 || ENGINE_MINOR_VERSION >= 27)
 	const USkeleton* Skeleton = SkeletalMesh->GetSkeleton();
 	const FReferenceSkeleton& RefSkeleton = SkeletalMesh->GetRefSkeleton();
-#else
-	const USkeleton* Skeleton = SkeletalMesh->Skeleton;
-	const FReferenceSkeleton& RefSkeleton = SkeletalMesh->RefSkeleton;
-#endif
 
 	FGLTFJsonSkin* JsonSkin = Builder.AddSkin();
 	JsonSkin->Name = Skeleton != nullptr ? Skeleton->GetName() : SkeletalMesh->GetName();
diff --git a/Engine/Plugins/Enterprise/GLTFExporter/Source/GLTFExporter/Private/Exporters/GLTFExporterUtility.cpp b/Engine/Plugins/Enterprise/GLTFExporter/Source/GLTFExporter/Private/Exporters/GLTFExporterUtility.cpp
--- a/Engine/Plugins/Enterprise/GLTFExporter/Source/GLTFExporter/Private/Exporters/GLTFExporterUtility.cpp
+++ b/Engine/Plugins/Enterprise/GLTFExporter/Source/GLTFExporter/Private/Exporters/GLTFExporterUtility.cpp
@@ -65,11 +65,7 @@ const USkeletalMesh* FGLTFExporterUtility::GetPreviewMesh(const UAnimSequence* A
 
 const USkeletalMesh* FGLTFExporterUtility::FindCompatibleMesh(const USkeleton *Skeleton)
 {
-#if (ENGINE_MAJOR_VERSION > 4 || ENGINE_MINOR_VERSION >= 27)
 	const FName SkeletonMemberName = USkeletalMesh::GetSkeletonMemberName();
-#else
-	const FName SkeletonMemberName = GET_MEMBER_NAME_CHECKED(USkeletalMesh, Skeleton);
-#endif
 
 	FARFilter Filter;
 	Filter.ClassPaths.Add(USkeletalMesh::StaticClass()->GetClassPathName());
diff --git a/Engine/Plugins/Enterprise/GLTFExporter/Source/GLTFMaterialBaking/Private/GLTFExportMaterialProxy.h b/Engine/Plugins/Enterprise/GLTFExporter/Source/GLTFMaterialBaking/Private/GLTFExportMaterialProxy.h
--- a/Engine/Plugins/Enterprise/GLTFExporter/Source/GLTFMaterialBaking/Private/GLTFExportMaterialProxy.h
+++ b/Engine/Plugins/Enterprise/GLTFExporter/Source/GLTFMaterialBaking/Private/GLTFExportMaterialProxy.h
@@ -120,18 +120,6 @@ struct FGLTFExportMaterialCompiler : public FProxyMaterialCompiler
 		return Compiler->Sub(Compiler->Mul(Compiler->Constant(2.0f), Compiler->Mul(Compiler->Dot(N, C), N)), C);
 	}
 
-#if (ENGINE_MAJOR_VERSION == 4 && ENGINE_MINOR_VERSION >= 26)
-	virtual int32 PreSkinVertexOffset() override
-	{
-		return Compiler->PreSkinVertexOffset();
-	}
-
-	virtual int32 PostSkinVertexOffset() override
-	{
-		return Compiler->PostSkinVertexOffset();
-	}
-#endif
-
 	virtual int32 PreSkinnedPosition() override
 	{
 		return Compiler->PreSkinnedPosition();
@@ -224,12 +212,7 @@ public:
 		, ProxyBlendMode(ProxyBlendMode)
 		, bTangentSpaceNormal(bTangentSpaceNormal)
 	{
-
-#if (ENGINE_MAJOR_VERSION > 4 || ENGINE_MINOR_VERSION >= 26)
 		SetQualityLevelProperties(GMaxRHIFeatureLevel);
-#else
-		SetQualityLevelProperties(EMaterialQualityLevel::High, false, GMaxRHIFeatureLevel);
-#endif
 		Material = InMaterialInterface->GetMaterial();
 		ReferencedTextures = InMaterialInterface->GetReferencedTextures();
 
@@ -248,11 +231,7 @@ public:
 			TArray<FVertexFactoryType*> VFTypes;
 			TArray<const FShaderPipelineType*> ShaderPipelineTypes;
 
-#if (ENGINE_MAJOR_VERSION > 4 || ENGINE_MINOR_VERSION >= 27)
 			GetDependentShaderAndVFTypes(GMaxRHIShaderPlatform, ResourceId.LayoutParams, ShaderTypes, ShaderPipelineTypes, VFTypes);
-#else
-			GetDependentShaderAndVFTypes(GMaxRHIShaderPlatform, ShaderTypes, ShaderPipelineTypes, VFTypes);
-#endif
 
 			// Overwrite the shader map Id's dependencies with ones that came from the FMaterial actually being compiled (this)
 			// This is necessary as we change FMaterial attributes like GetShadingModels(), which factor into the ShouldCache functions that determine dependent shader types
@@ -274,8 +253,6 @@ public:
 		case MP_Opacity: ResourceId.Usage = EMaterialShaderMapUsage::MaterialExportOpacity; break;
 		case MP_OpacityMask: ResourceId.Usage = EMaterialShaderMapUsage::MaterialExportOpacityMask; break;
 		case MP_SubsurfaceColor: ResourceId.Usage = EMaterialShaderMapUsage::MaterialExportSubSurfaceColor; break;
-
-#if (ENGINE_MAJOR_VERSION > 4 || ENGINE_MINOR_VERSION >= 27)
 		case MP_CustomData0: ResourceId.Usage = EMaterialShaderMapUsage::MaterialExportClearCoat; break;
 		case MP_CustomData1: ResourceId.Usage = EMaterialShaderMapUsage::MaterialExportClearCoatRoughness; break;
 		case MP_CustomOutput:
@@ -284,14 +261,6 @@ public:
 			break;
 		case MP_Refraction: ResourceId.Usage = static_cast<EMaterialShaderMapUsage::Type>(EMaterialShaderMapUsage::MaterialExportCustomOutput + 1); break;
 		case MP_ShadingModel: ResourceId.Usage = static_cast<EMaterialShaderMapUsage::Type>(EMaterialShaderMapUsage::MaterialExportCustomOutput + 2); break;
-#else
-		// NOTE: the following cases for custom data 0/1 and custom output are hacks since a proper solution requires engine changes:
-		case MP_CustomData0: ResourceId.Usage = static_cast<EMaterialShaderMapUsage::Type>(EMaterialShaderMapUsage::DebugViewMode + 1); break;
-		case MP_CustomData1: ResourceId.Usage = static_cast<EMaterialShaderMapUsage::Type>(EMaterialShaderMapUsage::DebugViewMode + 2); break;
-		case MP_CustomOutput: ResourceId.Usage = static_cast<EMaterialShaderMapUsage::Type>(EMaterialShaderMapUsage::DebugViewMode + 3); break;
-		case MP_Refraction: ResourceId.Usage = static_cast<EMaterialShaderMapUsage::Type>(EMaterialShaderMapUsage::DebugViewMode + 4); break;
-		case MP_ShadingModel: ResourceId.Usage = static_cast<EMaterialShaderMapUsage::Type>(EMaterialShaderMapUsage::DebugViewMode + 5); break;
-#endif
 
 		default:
 			ensureMsgf(false, TEXT("ExportMaterial has no usage for property %i.  Will likely reuse the normal rendering shader and crash later with a parameter mismatch"), (int32)InPropertyToCompile);
@@ -338,7 +307,6 @@ public:
 	////////////////
 	// FMaterialRenderProxy interface.
 
-#if (ENGINE_MAJOR_VERSION > 4 || ENGINE_MINOR_VERSION >= 27)
 	virtual const FMaterial* GetMaterialNoFallback(ERHIFeatureLevel::Type InFeatureLevel) const override
 	{
 		if (GetRenderingThreadShaderMap())
@@ -354,20 +322,6 @@ public:
 	}
 
 	virtual enum EStrataBlendMode GetStrataBlendMode() const override { return EStrataBlendMode::SBM_Opaque; }
-#else
-	virtual const FMaterial& GetMaterialWithFallback(ERHIFeatureLevel::Type FeatureLevel, const FMaterialRenderProxy*& OutFallbackMaterialRenderProxy) const override
-	{
-		if(GetRenderingThreadShaderMap())
-		{
-			return *this;
-		}
-		else
-		{
-			OutFallbackMaterialRenderProxy = UMaterial::GetDefaultMaterial(MD_Surface)->GetRenderProxy();
-			return OutFallbackMaterialRenderProxy->GetMaterialWithFallback(FeatureLevel, OutFallbackMaterialRenderProxy);
-		}
-	}
-#endif
 
 	virtual bool GetParameterValue(EMaterialParameterType Type, const FHashedMaterialParameterInfo& ParameterInfo, FMaterialParameterValue& OutValue, const FMaterialRenderContext& Context) const override
 	{
