Updated GLTFExporter to use PCHs and unity files.

Before:
Total CPU Time: 3465.375 s
Total time in Parallel executor: 147.16 seconds

After:
Total CPU Time: 47.90625 s
Total time in Parallel executor: 24.85 seconds

diff --git a/Engine/Plugins/Enterprise/GLTFExporter/Source/GLTFExporter/GLTFExporter.Build.cs b/Engine/Plugins/Enterprise/GLTFExporter/Source/GLTFExporter/GLTFExporter.Build.cs
--- a/Engine/Plugins/Enterprise/GLTFExporter/Source/GLTFExporter/GLTFExporter.Build.cs
+++ b/Engine/Plugins/Enterprise/GLTFExporter/Source/GLTFExporter/GLTFExporter.Build.cs
@@ -6,9 +6,7 @@ public class GLTFExporter : ModuleRules
 {
 	public GLTFExporter(ReadOnlyTargetRules Target) : base(Target)
 	{
-		//PCHUsage = PCHUsageMode.UseExplicitOrSharedPCHs;
-		PCHUsage = PCHUsageMode.NoPCHs;
-		bUseUnity = false;
+		PCHUsage = PCHUsageMode.UseExplicitOrSharedPCHs;
 
 		PublicDependencyModuleNames.AddRange(
 			new string[]
