Turned on PCH and unity file support for GLTFMaterialBaking and GLTFMaterialAnalyzer.

GLTFMaterialAnalyzer:
Before:
Total CPU Time: 46.21875 s
Total time in Parallel executor: 19.10 seconds
After:
Total CPU Time: 6.40625 s
Total time in Parallel executor: 17.17 seconds

GLTFMaterialBaking:
Before:
Total CPU Time: 102.0625 s
Total time in Parallel executor: 31.06 seconds
After:
Total CPU Time: 17.9375 s
Total time in Parallel executor: 10.23 seconds

diff --git a/Engine/Plugins/Enterprise/GLTFExporter/Source/GLTFMaterialAnalyzer/GLTFMaterialAnalyzer.Build.cs b/Engine/Plugins/Enterprise/GLTFExporter/Source/GLTFMaterialAnalyzer/GLTFMaterialAnalyzer.Build.cs
--- a/Engine/Plugins/Enterprise/GLTFExporter/Source/GLTFMaterialAnalyzer/GLTFMaterialAnalyzer.Build.cs
+++ b/Engine/Plugins/Enterprise/GLTFExporter/Source/GLTFMaterialAnalyzer/GLTFMaterialAnalyzer.Build.cs
@@ -6,9 +6,7 @@ public class GLTFMaterialAnalyzer : ModuleRules
 {
 	public GLTFMaterialAnalyzer(ReadOnlyTargetRules Target) : base(Target)
 	{
-		//PCHUsage = PCHUsageMode.UseExplicitOrSharedPCHs;
-		PCHUsage = PCHUsageMode.NoPCHs;
-		bUseUnity = false;
+		PCHUsage = PCHUsageMode.UseExplicitOrSharedPCHs;
 
 		PublicDependencyModuleNames.AddRange(
 			new string[]
diff --git a/Engine/Plugins/Enterprise/GLTFExporter/Source/GLTFMaterialBaking/GLTFMaterialBaking.Build.cs b/Engine/Plugins/Enterprise/GLTFExporter/Source/GLTFMaterialBaking/GLTFMaterialBaking.Build.cs
--- a/Engine/Plugins/Enterprise/GLTFExporter/Source/GLTFMaterialBaking/GLTFMaterialBaking.Build.cs
+++ b/Engine/Plugins/Enterprise/GLTFExporter/Source/GLTFMaterialBaking/GLTFMaterialBaking.Build.cs
@@ -6,9 +6,7 @@ public class GLTFMaterialBaking : ModuleRules
 {
 	public GLTFMaterialBaking(ReadOnlyTargetRules Target) : base(Target)
 	{
-		//PCHUsage = PCHUsageMode.UseExplicitOrSharedPCHs;
-		PCHUsage = PCHUsageMode.NoPCHs;
-		bUseUnity = false;
+		PCHUsage = PCHUsageMode.UseExplicitOrSharedPCHs;
 
 		PublicDependencyModuleNames.AddRange(
 			new string[]
