Add analytics for glTF exporter

diff --git a/Engine/Plugins/Enterprise/GLTFExporter/Source/GLTFExporter/GLTFExporter.Build.cs b/Engine/Plugins/Enterprise/GLTFExporter/Source/GLTFExporter/GLTFExporter.Build.cs
--- a/Engine/Plugins/Enterprise/GLTFExporter/Source/GLTFExporter/GLTFExporter.Build.cs
+++ b/Engine/Plugins/Enterprise/GLTFExporter/Source/GLTFExporter/GLTFExporter.Build.cs
@@ -29,6 +29,7 @@ public class GLTFExporter : ModuleRules
 				"MovieSceneTracks",
 				"VariantManagerContent",
 				"Projects",
+				"Analytics",
 			}
 		);
 
diff --git a/Engine/Plugins/Enterprise/GLTFExporter/Source/GLTFExporter/Private/Exporters/GLTFExporter.cpp b/Engine/Plugins/Enterprise/GLTFExporter/Source/GLTFExporter/Private/Exporters/GLTFExporter.cpp
--- a/Engine/Plugins/Enterprise/GLTFExporter/Source/GLTFExporter/Private/Exporters/GLTFExporter.cpp
+++ b/Engine/Plugins/Enterprise/GLTFExporter/Source/GLTFExporter/Private/Exporters/GLTFExporter.cpp
@@ -1,12 +1,13 @@
 // Copyright Epic Games, Inc. All Rights Reserved.
 
 #include "Exporters/GLTFExporter.h"
+#include "Exporters/GLTFExporterUtility.h"
+#include "Exporters/GLTFExporterAnalytics.h"
 #include "Options/GLTFExportOptions.h"
 #include "UI/GLTFExportOptionsWindow.h"
 #include "Builders/GLTFContainerBuilder.h"
 #include "UObject/GCObjectScopeGuard.h"
 #include "AssetExportTask.h"
-#include "GLTFExporterUtility.h"
 
 UGLTFExporter::UGLTFExporter(const FObjectInitializer& ObjectInitializer)
 	: Super(ObjectInitializer)
@@ -31,6 +32,8 @@ bool UGLTFExporter::ExportBinary(UObject* Object, const TCHAR* Type, FArchive& A
 		return false;
 	}
 
+	const uint64 StartTime = FPlatformTime::Cycles64();
+
 	FGCObjectScopeGuard OptionsGuard(Options);
 
 	TSet<AActor*> SelectedActors;
@@ -56,6 +59,8 @@ bool UGLTFExporter::ExportBinary(UObject* Object, const TCHAR* Type, FArchive& A
 		bSuccess = Builder.WriteExternalFiles(DirPath);
 	}
 
+	FGLTFExporterAnalytics(Object, Builder, true, IsAutomated(), bSuccess, StartTime).Send();
+
 	// TODO: should we copy messages to UAssetExportTask::Errors?
 
 	if (!FApp::IsUnattended() && Builder.HasLoggedMessages())
@@ -81,6 +86,8 @@ bool UGLTFExporter::ExportToGLTF(UObject* Object, const FString& FilePath, const
 		return false;
 	}
 
+	const uint64 StartTime = FPlatformTime::Cycles64();
+
 	const FString FileName = FPaths::GetCleanFilename(FilePath);
 	const FString DirPath = FPaths::GetPath(FilePath);
 
@@ -97,6 +104,8 @@ bool UGLTFExporter::ExportToGLTF(UObject* Object, const FString& FilePath, const
 		bSuccess = Builder.WriteAllFiles(DirPath);
 	}
 
+	FGLTFExporterAnalytics(Object, Builder, false, FApp::IsUnattended(), bSuccess, StartTime).Send();
+
 	return bSuccess;
 }
 
@@ -114,16 +123,10 @@ bool UGLTFExporter::AddObject(FGLTFContainerBuilder& Builder, const UObject* Obj
 const UGLTFExportOptions* UGLTFExporter::GetExportOptions()
 {
 	UGLTFExportOptions* Options = nullptr;
-	bool bAutomatedTask = GIsAutomationTesting || FApp::IsUnattended();
 
 	if (ExportTask != nullptr)
 	{
 		Options = Cast<UGLTFExportOptions>(ExportTask->Options);
-
-		if (ExportTask->bAutomated)
-		{
-			bAutomatedTask = true;
-		}
 	}
 
 	if (Options == nullptr)
@@ -132,7 +135,7 @@ const UGLTFExportOptions* UGLTFExporter::GetExportOptions()
 	}
 
 #if WITH_EDITOR
-	if (GetShowExportOption() && !bAutomatedTask)
+	if (GetShowExportOption() && !IsAutomated())
 	{
 		bool bExportAll = GetBatchMode();
 		bool bOperationCanceled = false;
@@ -158,3 +161,9 @@ FString UGLTFExporter::GetFilePath() const
 {
 	return ExportTask != nullptr ? ExportTask->Filename : CurrentFilename;
 }
+
+bool UGLTFExporter::IsAutomated() const
+{
+	// NOTE: FApp::IsUnattended() already checks GIsAutomationTesting
+	return FApp::IsUnattended() || (ExportTask != nullptr && ExportTask->bAutomated);
+}
diff --git a/Engine/Plugins/Enterprise/GLTFExporter/Source/GLTFExporter/Private/Exporters/GLTFExporterAnalytics.cpp b/Engine/Plugins/Enterprise/GLTFExporter/Source/GLTFExporter/Private/Exporters/GLTFExporterAnalytics.cpp
new file mode 100644
--- /dev/null
+++ b/Engine/Plugins/Enterprise/GLTFExporter/Source/GLTFExporter/Private/Exporters/GLTFExporterAnalytics.cpp
@@ -0,0 +1,135 @@
+// Copyright Epic Games, Inc. All Rights Reserved.
+
+#include "Exporters/GLTFExporterAnalytics.h"
+#include "Builders/GLTFConvertBuilder.h"
+#include "GLTFExporterModule.h"
+#include "EngineAnalytics.h"
+
+FGLTFExporterAnalytics::FGLTFExporterAnalytics(const UObject* Object, const FGLTFConvertBuilder& Builder, bool bInitiatedByTask, bool bExportAutomated, bool bExportSuccessful, uint64 StartTime)
+	: FGLTFExporterAnalytics(Object, Builder, bInitiatedByTask, bExportAutomated, bExportSuccessful, StartTime, FPlatformTime::Cycles64())
+{
+}
+
+FGLTFExporterAnalytics::FGLTFExporterAnalytics(const UObject* Object, const FGLTFConvertBuilder& Builder, bool bInitiatedByTask, bool bExportAutomated, bool bExportSuccessful, uint64 StartTime, uint64 EndTime)
+	: AssetType(Object != nullptr ? Object->GetClass() : nullptr)
+	, ExportOptions(Builder.ExportOptions)
+	, bExportAsGLB(Builder.bIsGLB)
+	, bSelectedOnly(!Builder.SelectedActors.IsEmpty())
+	, bInitiatedByTask(bInitiatedByTask)
+	, bExportAutomated(bExportAutomated)
+	, bExportSuccessful(bExportSuccessful)
+	, ExportDuration(FPlatformTime::ToSeconds64(EndTime - StartTime))
+{
+}
+
+void FGLTFExporterAnalytics::Send() const
+{
+	if (!FEngineAnalytics::IsAvailable())
+	{
+		return;
+	}
+
+	TArray<FAnalyticsEventAttribute> EventAttributes;
+	EventAttributes.Emplace(TEXT("AssetType"), AssetType != nullptr ? AssetType->GetName() : "null");
+	if (AssetType == UWorld::StaticClass()) EventAttributes.Emplace(TEXT("bSelectedOnly"), bSelectedOnly);
+	EventAttributes.Emplace(TEXT("bExportAsGLB"), bExportAsGLB);
+	EventAttributes.Emplace(TEXT("bInitiatedByTask"), bInitiatedByTask);
+	EventAttributes.Emplace(TEXT("bExportAutomated"), bExportAutomated);
+	EventAttributes.Emplace(TEXT("ExportDuration"), ExportDuration);
+
+	GetAttributesFromStruct(ExportOptions->GetClass(), ExportOptions, TEXT("ExportOptions"), EventAttributes);
+
+	EventAttributes.Emplace(TEXT("Platform"), FPlatformProperties::IniPlatformName());
+	EventAttributes.Emplace(TEXT("EngineMode"), FPlatformMisc::GetEngineMode());
+
+	const FString EventName = bExportSuccessful ? TEXT("GLTFExporter.Export") : TEXT("GLTFExporter.ExportFailure");
+	FEngineAnalytics::GetProvider().RecordEvent(EventName, EventAttributes);
+}
+
+void FGLTFExporterAnalytics::GetAttributesFromStruct(const UStruct* Struct, const void* ContainerPtr, const FString& AttributeName, TArray<FAnalyticsEventAttribute>& OutAttributes)
+{
+	for (const FProperty* Property = Struct->PropertyLink; Property != nullptr; Property = Property->PropertyLinkNext)
+	{
+		if (Property->HasAnyPropertyFlags(CPF_Deprecated | CPF_Transient))
+		{
+			continue; // Ignore deprecated and transient property
+		}
+
+#if WITH_EDITORONLY_DATA
+		if (Property->HasMetaData("InlineEditConditionToggle"))
+		{
+			continue; // Ignore inlined toggle property
+		}
+
+		FString EditCondition = Property->GetMetaData("EditCondition");
+		if (!EditCondition.IsEmpty())
+		{
+			if (const FBoolProperty* EditProperty = CastField<FBoolProperty>(Struct->FindPropertyByName(*EditCondition)))
+			{
+				const void* ValuePtr = EditProperty->ContainerPtrToValuePtr<void>(ContainerPtr);
+				if (!EditProperty->GetPropertyValue(ValuePtr))
+				{
+					continue; // Ignore property if (simple) edit condition is not met
+				}
+			}
+		}
+#endif
+
+		const void* ValuePtr = Property->ContainerPtrToValuePtr<void>(ContainerPtr);
+		GetAttributesFromProperty(Property, ValuePtr, AttributeName + TEXT(".") + Property->GetAuthoredName(), OutAttributes);
+	}
+}
+
+void FGLTFExporterAnalytics::GetAttributesFromProperty(const FProperty* Property, const void* ValuePtr, const FString& AttributeName, TArray<FAnalyticsEventAttribute>& OutAttributes)
+{
+	if (Property->ArrayDim != 1)
+	{
+		// If we get this warning, it means a new export option has been added that GLTFExporterAnalytics doesn't yet know how to serialize
+		UE_LOG(LogGLTFExporter, Warning, TEXT("GLTFExporterAnalytics doesn't support attribute %s with array dimension %d"), *AttributeName, Property->ArrayDim);
+		return;
+	}
+
+	if (const FBoolProperty* BoolProperty = CastField<FBoolProperty>(Property))
+	{
+		OutAttributes.Emplace(AttributeName, BoolProperty->GetPropertyValue(ValuePtr));
+	}
+	else if (const FStrProperty* StringProperty = CastField<FStrProperty>(Property))
+	{
+		OutAttributes.Emplace(AttributeName, StringProperty->GetPropertyValue(ValuePtr));
+	}
+	else if (Property->IsA<FEnumProperty>() || Property->IsA<FNumericProperty>())
+	{
+		FString ValueString;
+		Property->ExportTextItem_Direct(ValueString, ValuePtr, NULL, NULL, PPF_None);
+		OutAttributes.Emplace(AttributeName, ValueString);
+	}
+	else if (const FMapProperty* MapProperty = CastField<FMapProperty>(Property))
+	{
+		FScriptMapHelper Helper(MapProperty, ValuePtr);
+		int32 RemainingPairs = Helper.Num();
+
+		OutAttributes.Emplace(AttributeName, RemainingPairs > 0);
+
+		for (int32 Index = 0; RemainingPairs > 0; ++Index)
+		{
+			if (Helper.IsValidIndex(Index))
+			{
+				FString KeyString;
+				MapProperty->KeyProp->ExportTextItem_Direct(KeyString, Helper.GetKeyPtr(Index), nullptr, nullptr, PPF_None);
+
+				GetAttributesFromProperty(MapProperty->ValueProp, Helper.GetValuePtr(Index), AttributeName + TEXT(".") + KeyString, OutAttributes);
+
+				--RemainingPairs;
+			}
+		}
+	}
+	else if (const FStructProperty* StructProperty = CastField<FStructProperty>(Property))
+	{
+		GetAttributesFromStruct(StructProperty->Struct, ValuePtr, AttributeName, OutAttributes);
+	}
+	else
+	{
+		// If we get this warning, it means a new export option has been added that GLTFExporterAnalytics doesn't yet know how to serialize
+		UE_LOG(LogGLTFExporter, Warning, TEXT("GLTFExporterAnalytics doesn't support attribute %s with property class %s"), *AttributeName, *Property->GetClass()->GetName());
+	}
+}
diff --git a/Engine/Plugins/Enterprise/GLTFExporter/Source/GLTFExporter/Private/Exporters/GLTFExporterAnalytics.h b/Engine/Plugins/Enterprise/GLTFExporter/Source/GLTFExporter/Private/Exporters/GLTFExporterAnalytics.h
new file mode 100644
--- /dev/null
+++ b/Engine/Plugins/Enterprise/GLTFExporter/Source/GLTFExporter/Private/Exporters/GLTFExporterAnalytics.h
@@ -0,0 +1,31 @@
+// Copyright Epic Games, Inc. All Rights Reserved.
+
+#pragma once
+
+#include "CoreMinimal.h"
+
+class UGLTFExportOptions;
+class FGLTFConvertBuilder;
+struct FAnalyticsEventAttribute;
+
+struct FGLTFExporterAnalytics
+{
+	const UClass* AssetType;
+	const UGLTFExportOptions* ExportOptions;
+	bool bExportAsGLB;
+	bool bSelectedOnly;
+	bool bInitiatedByTask;
+	bool bExportAutomated;
+	bool bExportSuccessful;
+	double ExportDuration;
+
+	FGLTFExporterAnalytics(const UObject* Object, const FGLTFConvertBuilder& Builder, bool bInitiatedByTask, bool bExportAutomated, bool bExportSuccessful, uint64 StartTime);
+	FGLTFExporterAnalytics(const UObject* Object, const FGLTFConvertBuilder& Builder, bool bInitiatedByTask, bool bExportAutomated, bool bExportSuccessful, uint64 StartTime, uint64 EndTime);
+
+	void Send() const;
+
+private:
+
+	static void GetAttributesFromStruct(const UStruct* Struct, const void* ContainerPtr, const FString& AttributeName, TArray<FAnalyticsEventAttribute>& OutAttributes);
+	static void GetAttributesFromProperty(const FProperty* Property, const void* ValuePtr, const FString& AttributeName, TArray<FAnalyticsEventAttribute>& OutAttributes);
+};
diff --git a/Engine/Plugins/Enterprise/GLTFExporter/Source/GLTFExporter/Public/Exporters/GLTFExporter.h b/Engine/Plugins/Enterprise/GLTFExporter/Source/GLTFExporter/Public/Exporters/GLTFExporter.h
--- a/Engine/Plugins/Enterprise/GLTFExporter/Source/GLTFExporter/Public/Exporters/GLTFExporter.h
+++ b/Engine/Plugins/Enterprise/GLTFExporter/Source/GLTFExporter/Public/Exporters/GLTFExporter.h
@@ -46,5 +46,7 @@ protected:
 private:
 
 	const UGLTFExportOptions* GetExportOptions();
+
 	FString GetFilePath() const;
+	bool IsAutomated() const;
 };
