Fix support for utf8 character encoding in glTF exporter

diff --git a/Engine/Plugins/Enterprise/GLTFExporter/Source/GLTFExporter/Private/Json/GLTFJsonWriter.cpp b/Engine/Plugins/Enterprise/GLTFExporter/Source/GLTFExporter/Private/Json/GLTFJsonWriter.cpp
--- a/Engine/Plugins/Enterprise/GLTFExporter/Source/GLTFExporter/Private/Json/GLTFJsonWriter.cpp
+++ b/Engine/Plugins/Enterprise/GLTFExporter/Source/GLTFExporter/Private/Json/GLTFJsonWriter.cpp
@@ -155,7 +155,7 @@ struct TGLTFJsonPrintPolicyUTF8
 	{
 		ANSICHAR Buffer[8] = { 0 };
 		const int32 Length = FTCHARToUTF8_Convert::Utf8FromCodepoint(Char, Buffer, UE_ARRAY_COUNT(Buffer));
-		Stream->Serialize(&Char, Length);
+		Stream->Serialize(Buffer, Length);
 	}
 
 	static void WriteString(FArchive* Stream, const FString& String)
