Remove support for creating proxy materials in runtime in GLTFExporter

diff --git a/Engine/Plugins/Enterprise/GLTFExporter/Source/GLTFExporter/Private/Converters/GLTFMaterialProxyFactory.cpp b/Engine/Plugins/Enterprise/GLTFExporter/Source/GLTFExporter/Private/Converters/GLTFMaterialProxyFactory.cpp
--- a/Engine/Plugins/Enterprise/GLTFExporter/Source/GLTFExporter/Private/Converters/GLTFMaterialProxyFactory.cpp
+++ b/Engine/Plugins/Enterprise/GLTFExporter/Source/GLTFExporter/Private/Converters/GLTFMaterialProxyFactory.cpp
@@ -207,7 +207,7 @@ UMaterialInstanceConstant* FGLTFMaterialProxyFactory::CreateInstancedMaterial(UM
 {
 	const FString BaseName = TEXT("MI_GLTF_") + OriginalMaterial->GetName();
 	UPackage* Package = FindOrCreatePackage(BaseName);
-	return FGLTFProxyMaterialUtilities::CreateProxyMaterial<UMaterialInstanceConstant>(ShadingModel, Package, *BaseName, RF_Public | RF_Standalone);
+	return FGLTFProxyMaterialUtilities::CreateProxyMaterial(ShadingModel, Package, *BaseName, RF_Public | RF_Standalone);
 }
 
 UPackage* FGLTFMaterialProxyFactory::FindOrCreatePackage(const FString& BaseName)
diff --git a/Engine/Plugins/Enterprise/GLTFExporter/Source/GLTFExporter/Private/GLTFExporterModule.cpp b/Engine/Plugins/Enterprise/GLTFExporter/Source/GLTFExporter/Private/GLTFExporterModule.cpp
--- a/Engine/Plugins/Enterprise/GLTFExporter/Source/GLTFExporter/Private/GLTFExporterModule.cpp
+++ b/Engine/Plugins/Enterprise/GLTFExporter/Source/GLTFExporter/Private/GLTFExporterModule.cpp
@@ -37,21 +37,9 @@ public:
 #if WITH_EDITOR
 	void PostEngineInit()
 	{
-		LoadProxyParentsForCook();
 		AddProxyAssetActions();
 	}
 
-	static void LoadProxyParentsForCook()
-	{
-		static FGLTFObjectArrayScopeGuard ScopeGuard;
-
-		constexpr int32 NumMaterials = static_cast<int32>(EGLTFJsonShadingModel::NumShadingModels);
-		for (int32 i = 0; i < NumMaterials; ++i)
-		{
-			ScopeGuard.Add(FGLTFProxyMaterialUtilities::GetBaseMaterial(static_cast<EGLTFJsonShadingModel>(i)));
-		}
-	}
-
 	void AddProxyAssetActions()
 	{
 		if (FAssetToolsModule::IsModuleLoaded())
diff --git a/Engine/Plugins/Enterprise/GLTFExporter/Source/GLTFExporter/Private/Materials/GLTFProxyMaterialInfo.h b/Engine/Plugins/Enterprise/GLTFExporter/Source/GLTFExporter/Private/Materials/GLTFProxyMaterialInfo.h
new file mode 100644
--- /dev/null
+++ b/Engine/Plugins/Enterprise/GLTFExporter/Source/GLTFExporter/Private/Materials/GLTFProxyMaterialInfo.h
@@ -0,0 +1,35 @@
+// Copyright Epic Games, Inc. All Rights Reserved.
+
+#pragma once
+
+#include "Materials/GLTFProxyMaterialParameterInfo.h"
+
+class GLTFEXPORTER_API FGLTFProxyMaterialInfo
+{
+public:
+
+	static const FGLTFProxyMaterialTextureParameterInfo BaseColor;
+	static const TGLTFProxyMaterialParameterInfo<FLinearColor> BaseColorFactor;
+
+	static const FGLTFProxyMaterialTextureParameterInfo Emissive;
+	static const TGLTFProxyMaterialParameterInfo<FLinearColor> EmissiveFactor;
+
+	static const FGLTFProxyMaterialTextureParameterInfo MetallicRoughness;
+	static const TGLTFProxyMaterialParameterInfo<float> MetallicFactor;
+	static const TGLTFProxyMaterialParameterInfo<float> RoughnessFactor;
+
+	static const FGLTFProxyMaterialTextureParameterInfo Normal;
+	static const TGLTFProxyMaterialParameterInfo<float> NormalScale;
+
+	static const FGLTFProxyMaterialTextureParameterInfo Occlusion;
+	static const TGLTFProxyMaterialParameterInfo<float> OcclusionStrength;
+
+	static const FGLTFProxyMaterialTextureParameterInfo ClearCoat;
+	static const TGLTFProxyMaterialParameterInfo<float> ClearCoatFactor;
+
+	static const FGLTFProxyMaterialTextureParameterInfo ClearCoatRoughness;
+	static const TGLTFProxyMaterialParameterInfo<float> ClearCoatRoughnessFactor;
+
+	static const FGLTFProxyMaterialTextureParameterInfo ClearCoatNormal;
+	static const TGLTFProxyMaterialParameterInfo<float> ClearCoatNormalScale;
+};
diff --git a/Engine/Plugins/Enterprise/GLTFExporter/Source/GLTFExporter/Private/Materials/GLTFProxyMaterialParameterInfo.h b/Engine/Plugins/Enterprise/GLTFExporter/Source/GLTFExporter/Private/Materials/GLTFProxyMaterialParameterInfo.h
new file mode 100644
--- /dev/null
+++ b/Engine/Plugins/Enterprise/GLTFExporter/Source/GLTFExporter/Private/Materials/GLTFProxyMaterialParameterInfo.h
@@ -0,0 +1,89 @@
+// Copyright Epic Games, Inc. All Rights Reserved.
+
+#pragma once
+
+#include "Utilities/GLTFProxyMaterialUtilities.h"
+
+class UMaterialInstanceConstant;
+class UMaterialInstanceDynamic;
+class UMaterialInterface;
+class UTexture;
+
+template <typename ParameterType>
+class TGLTFProxyMaterialParameterInfo
+{
+public:
+
+	TGLTFProxyMaterialParameterInfo(const FString& ParameterName)
+		: ParameterInfo(*ParameterName)
+	{
+	}
+
+	bool Get(const UMaterialInterface* Material, ParameterType& OutValue, bool NonDefaultOnly = false) const
+	{
+		return FGLTFProxyMaterialUtilities::GetParameterValue(Material, ParameterInfo, OutValue, NonDefaultOnly);
+	}
+
+	void Set(UMaterialInstanceDynamic* Material, const ParameterType& Value, bool NonDefaultOnly = false) const
+	{
+		FGLTFProxyMaterialUtilities::SetParameterValue(Material, ParameterInfo, Value, NonDefaultOnly);
+	}
+
+#if WITH_EDITOR
+	void Set(UMaterialInstanceConstant* Material, const ParameterType& Value, bool NonDefaultOnly = false) const
+	{
+		FGLTFProxyMaterialUtilities::SetParameterValue(Material, ParameterInfo, Value, NonDefaultOnly);
+	}
+#endif
+
+	FString ToString() const
+	{
+		return ParameterInfo.Name.ToString();
+	}
+
+	bool operator==(const TGLTFProxyMaterialParameterInfo<ParameterType>& Other) const
+	{
+		return ParameterInfo.Name == Other.ParameterInfo.Name
+			&& ParameterInfo.Index == Other.ParameterInfo.Index
+			&& ParameterInfo.Association == Other.ParameterInfo.Association;
+	}
+
+	bool operator!=(const TGLTFProxyMaterialParameterInfo<ParameterType>& Other) const
+	{
+		return !(*this == Other);
+	}
+
+private:
+
+	FHashedMaterialParameterInfo ParameterInfo;
+};
+
+class FGLTFProxyMaterialTextureParameterInfo
+{
+public:
+
+	FGLTFProxyMaterialTextureParameterInfo(const FString& ParameterName)
+		: Texture(ParameterName + TEXT(" Texture"))
+		, UVIndex(ParameterName + TEXT(" UV Index"))
+		, UVOffset(ParameterName + TEXT(" UV Offset"))
+		, UVScale(ParameterName + TEXT(" UV Scale"))
+		, UVRotation(ParameterName + TEXT(" UV Rotation"))
+	{
+	}
+
+	bool operator==(const FGLTFProxyMaterialTextureParameterInfo& Other) const
+	{
+		return Texture == Other.Texture;
+	}
+
+	bool operator!=(const FGLTFProxyMaterialTextureParameterInfo& Other) const
+	{
+		return !(*this == Other);
+	}
+
+	TGLTFProxyMaterialParameterInfo<UTexture*> Texture;
+	TGLTFProxyMaterialParameterInfo<float> UVIndex;
+	TGLTFProxyMaterialParameterInfo<FLinearColor> UVOffset;
+	TGLTFProxyMaterialParameterInfo<FLinearColor> UVScale;
+	TGLTFProxyMaterialParameterInfo<float> UVRotation;
+};
diff --git a/Engine/Plugins/Enterprise/GLTFExporter/Source/GLTFExporter/Private/Utilities/GLTFProxyMaterialUtilities.cpp b/Engine/Plugins/Enterprise/GLTFExporter/Source/GLTFExporter/Private/Utilities/GLTFProxyMaterialUtilities.cpp
--- a/Engine/Plugins/Enterprise/GLTFExporter/Source/GLTFExporter/Private/Utilities/GLTFProxyMaterialUtilities.cpp
+++ b/Engine/Plugins/Enterprise/GLTFExporter/Source/GLTFExporter/Private/Utilities/GLTFProxyMaterialUtilities.cpp
@@ -3,12 +3,21 @@
 #include "Utilities/GLTFProxyMaterialUtilities.h"
 #include "Utilities/GLTFJsonUtilities.h"
 #include "UserData/GLTFMaterialUserData.h"
-#include "Materials/MaterialInstanceDynamic.h"
 #include "Materials/MaterialInstanceConstant.h"
 
 #define PROXY_MATERIAL_NAME_PREFIX TEXT("M_GLTF_")
 #define PROXY_MATERIAL_ROOT_PATH   TEXT("/GLTFExporter/Materials/Proxy/")
 
+#if WITH_EDITOR
+UMaterialInstanceConstant* FGLTFProxyMaterialUtilities::CreateProxyMaterial(EGLTFJsonShadingModel ShadingModel,
+	UObject* Outer, FName Name, EObjectFlags Flags)
+{
+	UMaterialInstanceConstant* MaterialInstance = NewObject<UMaterialInstanceConstant>(Outer, Name, Flags);
+	MaterialInstance->Parent = GetBaseMaterial(ShadingModel);
+	return MaterialInstance;
+}
+#endif
+
 bool FGLTFProxyMaterialUtilities::IsProxyMaterial(const UMaterial* Material)
 {
 	return Material != nullptr && Material->GetPathName().StartsWith(PROXY_MATERIAL_ROOT_PATH PROXY_MATERIAL_NAME_PREFIX);
@@ -126,48 +135,6 @@ bool FGLTFProxyMaterialUtilities::GetParameterValue(const UMaterialInterface* Ma
 	return true;
 }
 
-void FGLTFProxyMaterialUtilities::SetParameterValue(UMaterialInstanceDynamic* Material, const FHashedMaterialParameterInfo& ParameterInfo, float Value, bool NonDefaultOnly)
-{
-	if (NonDefaultOnly)
-	{
-		float DefaultValue;
-		if (!Material->GetScalarParameterDefaultValue(ParameterInfo, DefaultValue) || DefaultValue == Value)
-		{
-			return;
-		}
-	}
-
-	Material->SetScalarParameterValueByInfo(FMaterialParameterInfo(ParameterInfo), Value);
-}
-
-void FGLTFProxyMaterialUtilities::SetParameterValue(UMaterialInstanceDynamic* Material, const FHashedMaterialParameterInfo& ParameterInfo, const FLinearColor& Value, bool NonDefaultOnly)
-{
-	if (NonDefaultOnly)
-	{
-		FLinearColor DefaultValue;
-		if (!Material->GetVectorParameterDefaultValue(ParameterInfo, DefaultValue) || DefaultValue == Value)
-		{
-			return;
-		}
-	}
-
-	Material->SetVectorParameterValueByInfo(FMaterialParameterInfo(ParameterInfo), Value);
-}
-
-void FGLTFProxyMaterialUtilities::SetParameterValue(UMaterialInstanceDynamic* Material, const FHashedMaterialParameterInfo& ParameterInfo, UTexture* Value, bool NonDefaultOnly)
-{
-	if (NonDefaultOnly)
-	{
-		UTexture* DefaultValue;
-		if (!Material->GetTextureParameterDefaultValue(ParameterInfo, DefaultValue) || DefaultValue == Value)
-		{
-			return;
-		}
-	}
-
-	Material->SetTextureParameterValueByInfo(FMaterialParameterInfo(ParameterInfo), Value);
-}
-
 #if WITH_EDITOR
 
 void FGLTFProxyMaterialUtilities::SetParameterValue(UMaterialInstanceConstant* Material, const FHashedMaterialParameterInfo& ParameterInfo, float Value, bool NonDefaultOnly)
@@ -214,10 +181,10 @@ void FGLTFProxyMaterialUtilities::SetParameterValue(UMaterialInstanceConstant* M
 
 #endif
 
-bool FGLTFProxyMaterialUtilities::GetTwoSided(const UMaterialInstance* Material, bool& OutValue, bool NonDefaultOnly)
+bool FGLTFProxyMaterialUtilities::GetTwoSided(const UMaterialInterface* Material, bool& OutValue, bool NonDefaultOnly)
 {
-	const bool DefaultValue = Material->Parent->IsTwoSided();
-	const bool Value = Material->BasePropertyOverrides.bOverride_TwoSided ? Material->BasePropertyOverrides.TwoSided : DefaultValue;
+	const bool DefaultValue = Material->GetMaterial()->IsTwoSided();
+	const bool Value = Material->IsTwoSided();
 
 	if (NonDefaultOnly && Value == DefaultValue)
 	{
@@ -228,10 +195,10 @@ bool FGLTFProxyMaterialUtilities::GetTwoSided(const UMaterialInstance* Material,
 	return true;
 }
 
-bool FGLTFProxyMaterialUtilities::GetBlendMode(const UMaterialInstance* Material, EBlendMode& OutValue, bool NonDefaultOnly)
+bool FGLTFProxyMaterialUtilities::GetBlendMode(const UMaterialInterface* Material, EBlendMode& OutValue, bool NonDefaultOnly)
 {
-	const EBlendMode DefaultValue = Material->Parent->GetBlendMode();
-	const EBlendMode Value = Material->BasePropertyOverrides.bOverride_BlendMode ? Material->BasePropertyOverrides.BlendMode.GetValue() : DefaultValue;
+	const EBlendMode DefaultValue = Material->GetMaterial()->GetBlendMode();
+	const EBlendMode Value = Material->GetBlendMode();
 
 	if (NonDefaultOnly && Value == DefaultValue)
 	{
@@ -242,10 +209,10 @@ bool FGLTFProxyMaterialUtilities::GetBlendMode(const UMaterialInstance* Material
 	return true;
 }
 
-bool FGLTFProxyMaterialUtilities::GetOpacityMaskClipValue(const UMaterialInstance* Material, float& OutValue, bool NonDefaultOnly)
+bool FGLTFProxyMaterialUtilities::GetOpacityMaskClipValue(const UMaterialInterface* Material, float& OutValue, bool NonDefaultOnly)
 {
-	const float DefaultValue = Material->Parent->GetOpacityMaskClipValue();
-	const float Value = Material->BasePropertyOverrides.bOverride_OpacityMaskClipValue ? Material->BasePropertyOverrides.OpacityMaskClipValue : DefaultValue;
+	const float DefaultValue = Material->GetMaterial()->GetOpacityMaskClipValue();
+	const float Value = Material->GetOpacityMaskClipValue();
 
 	if (NonDefaultOnly && Value == DefaultValue)
 	{
@@ -256,10 +223,9 @@ bool FGLTFProxyMaterialUtilities::GetOpacityMaskClipValue(const UMaterialInstanc
 	return true;
 }
 
-void FGLTFProxyMaterialUtilities::SetTwoSided(UMaterialInstance* Material, bool Value, bool NonDefaultOnly)
+#if WITH_EDITOR
+void FGLTFProxyMaterialUtilities::SetTwoSided(UMaterialInstanceConstant* Material, bool Value, bool NonDefaultOnly)
 {
-	// TODO: Even if the exporter supports it, UMaterialInstanceDynamic won't actually reflect these overriden base properties in Unreal
-
 	if (NonDefaultOnly)
 	{
 		const bool DefaultValue = Material->Parent->IsTwoSided();
@@ -274,10 +240,8 @@ void FGLTFProxyMaterialUtilities::SetTwoSided(UMaterialInstance* Material, bool
 	Material->TwoSided = Value;
 }
 
-void FGLTFProxyMaterialUtilities::SetBlendMode(UMaterialInstance* Material, EBlendMode Value, bool NonDefaultOnly)
+void FGLTFProxyMaterialUtilities::SetBlendMode(UMaterialInstanceConstant* Material, EBlendMode Value, bool NonDefaultOnly)
 {
-	// TODO: Even if the exporter supports it, UMaterialInstanceDynamic won't actually reflect these overriden base properties in Unreal
-
 	if (NonDefaultOnly)
 	{
 		const EBlendMode DefaultValue = Material->Parent->GetBlendMode();
@@ -292,10 +256,8 @@ void FGLTFProxyMaterialUtilities::SetBlendMode(UMaterialInstance* Material, EBle
 	Material->BlendMode = Value;
 }
 
-void FGLTFProxyMaterialUtilities::SetOpacityMaskClipValue(UMaterialInstance* Material, float Value, bool NonDefaultOnly)
+void FGLTFProxyMaterialUtilities::SetOpacityMaskClipValue(UMaterialInstanceConstant* Material, float Value, bool NonDefaultOnly)
 {
-	// TODO: Even if the exporter supports it, UMaterialInstanceDynamic won't actually reflect these overriden base properties in Unreal
-
 	if (NonDefaultOnly)
 	{
 		const float DefaultValue = Material->Parent->GetOpacityMaskClipValue();
@@ -309,3 +271,4 @@ void FGLTFProxyMaterialUtilities::SetOpacityMaskClipValue(UMaterialInstance* Mat
 	Material->BasePropertyOverrides.OpacityMaskClipValue = Value;
 	Material->OpacityMaskClipValue = Value;
 }
+#endif
diff --git a/Engine/Plugins/Enterprise/GLTFExporter/Source/GLTFExporter/Public/Materials/GLTFProxyMaterial.h b/Engine/Plugins/Enterprise/GLTFExporter/Source/GLTFExporter/Public/Materials/GLTFProxyMaterial.h
deleted file mode 100644
diff --git a/Engine/Plugins/Enterprise/GLTFExporter/Source/GLTFExporter/Public/Materials/GLTFProxyMaterialInfo.h b/Engine/Plugins/Enterprise/GLTFExporter/Source/GLTFExporter/Public/Materials/GLTFProxyMaterialInfo.h
deleted file mode 100644
diff --git a/Engine/Plugins/Enterprise/GLTFExporter/Source/GLTFExporter/Public/Materials/GLTFProxyMaterialParameter.h b/Engine/Plugins/Enterprise/GLTFExporter/Source/GLTFExporter/Public/Materials/GLTFProxyMaterialParameter.h
deleted file mode 100644
diff --git a/Engine/Plugins/Enterprise/GLTFExporter/Source/GLTFExporter/Public/Materials/GLTFProxyMaterialParameterInfo.h b/Engine/Plugins/Enterprise/GLTFExporter/Source/GLTFExporter/Public/Materials/GLTFProxyMaterialParameterInfo.h
deleted file mode 100644
diff --git a/Engine/Plugins/Enterprise/GLTFExporter/Source/GLTFExporter/Public/Utilities/GLTFProxyMaterialUtilities.h b/Engine/Plugins/Enterprise/GLTFExporter/Source/GLTFExporter/Public/Utilities/GLTFProxyMaterialUtilities.h
--- a/Engine/Plugins/Enterprise/GLTFExporter/Source/GLTFExporter/Public/Utilities/GLTFProxyMaterialUtilities.h
+++ b/Engine/Plugins/Enterprise/GLTFExporter/Source/GLTFExporter/Public/Utilities/GLTFProxyMaterialUtilities.h
@@ -8,23 +8,17 @@
 #include "UObject/ObjectMacros.h"
 
 class UMaterial;
-class UMaterialInstance;
-class UMaterialInstanceDynamic;
-class UMaterialInstanceConstant;
 class UMaterialInterface;
+class UMaterialInstanceConstant;
 class UTexture;
 
 class GLTFEXPORTER_API FGLTFProxyMaterialUtilities
 {
 public:
 
-	template <typename MaterialType, typename = typename TEnableIf<TIsDerivedFrom<MaterialType, UMaterialInstance>::Value>::Type>
-	static MaterialType* CreateProxyMaterial(EGLTFJsonShadingModel ShadingModel, UObject* Outer = GetTransientPackage(), FName Name = NAME_None, EObjectFlags Flags = RF_NoFlags)
-	{
-		MaterialType* MaterialInstance = NewObject<MaterialType>(Outer, Name, Flags);
-		MaterialInstance->Parent = GetBaseMaterial(ShadingModel);
-		return MaterialInstance;
-	}
+#if WITH_EDITOR
+	static UMaterialInstanceConstant* CreateProxyMaterial(EGLTFJsonShadingModel ShadingModel, UObject* Outer = GetTransientPackage(), FName Name = NAME_None, EObjectFlags Flags = RF_NoFlags);
+#endif
 
 	static bool IsProxyMaterial(const UMaterial* Material);
 	static bool IsProxyMaterial(const UMaterialInterface* Material);
@@ -38,21 +32,19 @@ public:
 	static bool GetParameterValue(const UMaterialInterface* Material, const FHashedMaterialParameterInfo& ParameterInfo, FLinearColor& OutValue, bool NonDefaultOnly = false);
 	static bool GetParameterValue(const UMaterialInterface* Material, const FHashedMaterialParameterInfo& ParameterInfo, UTexture*& OutValue, bool NonDefaultOnly = false);
 
-	static void SetParameterValue(UMaterialInstanceDynamic* Material, const FHashedMaterialParameterInfo& ParameterInfo, float Value, bool NonDefaultOnly = false);
-	static void SetParameterValue(UMaterialInstanceDynamic* Material, const FHashedMaterialParameterInfo& ParameterInfo, const FLinearColor& Value, bool NonDefaultOnly = false);
-	static void SetParameterValue(UMaterialInstanceDynamic* Material, const FHashedMaterialParameterInfo& ParameterInfo, UTexture* Value, bool NonDefaultOnly = false);
-
 #if WITH_EDITOR
 	static void SetParameterValue(UMaterialInstanceConstant* Material, const FHashedMaterialParameterInfo& ParameterInfo, float Value, bool NonDefaultOnly = false);
 	static void SetParameterValue(UMaterialInstanceConstant* Material, const FHashedMaterialParameterInfo& ParameterInfo, const FLinearColor& Value, bool NonDefaultOnly = false);
 	static void SetParameterValue(UMaterialInstanceConstant* Material, const FHashedMaterialParameterInfo& ParameterInfo, UTexture* Value, bool NonDefaultOnly = false);
 #endif
 
-	static bool GetTwoSided(const UMaterialInstance* Material, bool& OutValue, bool NonDefaultOnly = false);
-	static bool GetBlendMode(const UMaterialInstance* Material, EBlendMode& OutValue, bool NonDefaultOnly = false);
-	static bool GetOpacityMaskClipValue(const UMaterialInstance* Material, float& OutValue, bool NonDefaultOnly = false);
+	static bool GetTwoSided(const UMaterialInterface* Material, bool& OutValue, bool NonDefaultOnly = false);
+	static bool GetBlendMode(const UMaterialInterface* Material, EBlendMode& OutValue, bool NonDefaultOnly = false);
+	static bool GetOpacityMaskClipValue(const UMaterialInterface* Material, float& OutValue, bool NonDefaultOnly = false);
 
-	static void SetTwoSided(UMaterialInstance* Material, bool Value, bool NonDefaultOnly = false);
-	static void SetBlendMode(UMaterialInstance* Material, EBlendMode Value, bool NonDefaultOnly = false);
-	static void SetOpacityMaskClipValue(UMaterialInstance* Material, float Value, bool NonDefaultOnly = false);
+#if WITH_EDITOR
+	static void SetTwoSided(UMaterialInstanceConstant* Material, bool Value, bool NonDefaultOnly = false);
+	static void SetBlendMode(UMaterialInstanceConstant* Material, EBlendMode Value, bool NonDefaultOnly = false);
+	static void SetOpacityMaskClipValue(UMaterialInstanceConstant* Material, float Value, bool NonDefaultOnly = false);
+#endif
 };
