Remove support for EPIC_animation_hotspots and EPIC_camera_controls in GLTFExporter

diff --git a/Engine/Plugins/Enterprise/GLTFExporter/Docs/export-options-reference.md b/Engine/Plugins/Enterprise/GLTFExporter/Docs/export-options-reference.md
--- a/Engine/Plugins/Enterprise/GLTFExporter/Docs/export-options-reference.md
+++ b/Engine/Plugins/Enterprise/GLTFExporter/Docs/export-options-reference.md
@@ -21,7 +21,7 @@ Option                         | Description
 `Export Vertex Skin Weights`   | If enabled, export vertex bone weights and indices in skeletal meshes. Necessary for animation sequences.
 `Use Mesh Quantization`        | If enabled, use quantization for vertex tangents and normals, reducing size. Requires extension KHR_mesh_quantization, which may result in the mesh not loading in some glTF viewers.
 `Export Level Sequences`       | If enabled, export level sequences. Only transform tracks are currently supported. The level sequence will be played at the assigned display rate.
-`Export Animation Sequences`   | If enabled, export single animation asset used by a skeletal mesh component or hotspot actor. Export of vertex skin weights must be enabled.
+`Export Animation Sequences`   | If enabled, export single animation asset used by a skeletal mesh component. Export of vertex skin weights must be enabled.
 `Export Playback Settings`     | If enabled, export play rate, start time, looping, and auto play for an animation or level sequence. Uses extension EPIC_animation_playback, which is supported by Unreal's glTF viewer.
 `Texture Image Format`         | Desired image format used for exported textures.
 `Texture Image Quality`        | Level of compression used for textures exported with lossy image formats, 0 (default) or value between 1 (worst quality, best compression) and 100 (best quality, worst compression).
@@ -33,8 +33,6 @@ Option                         | Description
 `Export Hidden In Game`        | If enabled, export actors and components that are flagged as hidden in-game.
 `Export Lights`                | Mobility of directional, point, and spot light components that will be exported. Uses extension KHR_lights_punctual.
 `Export Cameras`               | If enabled, export camera components.
-`Export Camera Controls`       | If enabled, export GLTFCameraActors. Uses extension EPIC_camera_controls, which is supported by Unreal's glTF viewer.
-`Export Animation Hotspots`    | If enabled, export GLTFHotspotActors. Uses extension EPIC_animation_hotspots, which is supported by Unreal's glTF viewer.
 `Export HDRI Backdrops`        | If enabled, export HDRIBackdrop blueprints. Uses extension EPIC_hdri_backdrops, which is supported by Unreal's glTF viewer.
 `Export Sky Spheres`           | If enabled, export SkySphere blueprints. Uses extension EPIC_sky_spheres, which is supported by Unreal's glTF viewer.
 `Variant Sets Mode`            | Mode determining if and how to export LevelVariantSetsActors.
diff --git a/Engine/Plugins/Enterprise/GLTFExporter/Docs/what-can-be-exported.md b/Engine/Plugins/Enterprise/GLTFExporter/Docs/what-can-be-exported.md
--- a/Engine/Plugins/Enterprise/GLTFExporter/Docs/what-can-be-exported.md
+++ b/Engine/Plugins/Enterprise/GLTFExporter/Docs/what-can-be-exported.md
@@ -222,8 +222,6 @@ Support for actors in a level, can be split into two categories:
   - [`Sky Sphere`](#sky-spheres)
   - [`Level Sequence Actor`](#level-sequence-actors)
   - [`Level Variant Sets Actor`](#level-variant-sets-actors)
-  - [`GLTF Camera`](#gltf-cameras)
-  - [`GLTF Hotspot`](#gltf-hotspots)
 
 - An actor that has one or more supported components.
 
@@ -288,43 +286,6 @@ Property         | Description
 
 `Level Variant Sets Actor` is provided by an engine-distributed plugin (`Variant Manager`) and ìs responsible for querying and activating variants in a `Level Variant Sets` at runtime. Only the `Level Variant Sets` property is exported.
 
-### GLTF Cameras
-
-> Please note that export of `GLTF Camera` uses the glTF extension `EPIC_camera_controls`, which can be turned off in the export options.
-
-`GLTF Camera` is an actor provided by the exporter plugin, to add user-controlled cameras to a glTF scene. It has the following properties that are supported by the exporter:
-
-Property               | Description
------------------------| ---------------------------------------------------------------------------------------------------------
-`Mode`                 | Camera mode which can be `Orbital` (rotate the camera around a target and zoom in/out) or `Free Look` (rotate the camera in first person but no movement).
-`Target`               | Actor which the camera will focus on and subsequently orbit when using Orbital mode.
-`Pitch Angle Min`      | Minimum pitch angle (in degrees) for the camera.
-`Pitch Angle Max`      | Maximum pitch angle (in degrees) for the camera.
-`Yaw Angle Min`        | Minimum yaw angle (in degrees) for the camera.
-`Yaw Angle Max`        | Maximum yaw angle (in degrees) for the camera.
-`Distance Min`         | Closest distance the camera can approach the focused actor.
-`Distance Max`         | Farthest distance the camera can recede from the focused actor.
-`Dolly Duration`       | Duration (in seconds) that it takes the camera to complete a change in distance.
-`Dolly Sensitivity`    | Size of the dolly movement relative to user input. The higher the value, the faster it moves.
-`Rotation Inertia`     | Deceleration that occurs after rotational movement. The higher the value, the longer it takes to settle.
-`Rotation Sensitivity` | Size of the rotational movement relative to user input. The higher the value, the faster it moves.
-
-### GLTF Hotspots
-
-> Please note that export of `GLTF Hotspot` uses the glTF extension `EPIC_animation_hotspots`, which can be turned off in the export options.
-
-`GLTF Hotspot` is an actor provided by the exporter plugin, to add a clickable UI element (that triggers animations) to a glTF scene. It has the following properties that are supported by the exporter:
-
-Property                | Description
-------------------------| ------------------------------------------------------------------------------------------------------------------------
-`Skeletal Mesh Actor`   | The skeletal mesh actor that will be animated when the hotspot is clicked.
-`Animation Sequence`    | The animation sequence that will be played on the skeletal mesh actor. Must be compatible with its skeletal mesh asset.
-`Level Sequence`        | The level sequence that will be played in the level.
-`Image`                 | The billboard image that will be shown when the hotspot is in an inactive state or one without a specified image.
-`Hovered Image`         | The optional billboard image that will be shown when a cursor enters the hotspot
-`Toggled Image`         | The optional billboard image that will be shown when the hotspot is toggled by a click.
-`Toggled Hovered Image` | The optional billboard image that will be shown when the hotspot is toggled by a click and a cursor enters it.
-
 ## Components
 
 Supported components can be separated into three groups:
diff --git a/Engine/Plugins/Enterprise/GLTFExporter/Docs/what-is-gltf.md b/Engine/Plugins/Enterprise/GLTFExporter/Docs/what-is-gltf.md
--- a/Engine/Plugins/Enterprise/GLTFExporter/Docs/what-is-gltf.md
+++ b/Engine/Plugins/Enterprise/GLTFExporter/Docs/what-is-gltf.md
@@ -40,8 +40,6 @@ Extension                   | Description
 `EPIC_level_variant_sets`   | Scene variants by Unreal's variant manager
 `EPIC_hdri_backdrops`       | Unreal backdrop actors for HDR image projection
 `EPIC_sky_spheres`          | Unreal's default sky sphere actors
-`EPIC_animation_hotspots`   | Exporter-provided actors for interactive animations
 `EPIC_animation_playback`   | Animation looping, auto play, start time, play rate
-`EPIC_camera_controls`      | Exporter-provided actors for user-controlled cameras
 `EPIC_texture_hdr_encoding` | HDR encoding for image formats like PNG and JPEG
 `EPIC_blend_modes`          | Additive, modulate, and alpha composite blend modes
diff --git a/Engine/Plugins/Enterprise/GLTFExporter/Source/GLTFExporter/Private/Actors/GLTFCameraActor.cpp b/Engine/Plugins/Enterprise/GLTFExporter/Source/GLTFExporter/Private/Actors/GLTFCameraActor.cpp
deleted file mode 100644
diff --git a/Engine/Plugins/Enterprise/GLTFExporter/Source/GLTFExporter/Private/Actors/GLTFHotspotActor.cpp b/Engine/Plugins/Enterprise/GLTFExporter/Source/GLTFExporter/Private/Actors/GLTFHotspotActor.cpp
deleted file mode 100644
diff --git a/Engine/Plugins/Enterprise/GLTFExporter/Source/GLTFExporter/Private/Builders/GLTFConvertBuilder.cpp b/Engine/Plugins/Enterprise/GLTFExporter/Source/GLTFExporter/Private/Builders/GLTFConvertBuilder.cpp
--- a/Engine/Plugins/Enterprise/GLTFExporter/Source/GLTFExporter/Private/Builders/GLTFConvertBuilder.cpp
+++ b/Engine/Plugins/Enterprise/GLTFExporter/Source/GLTFExporter/Private/Builders/GLTFConvertBuilder.cpp
@@ -523,16 +523,6 @@ FGLTFJsonLightMap* FGLTFConvertBuilder::AddUniqueLightMap(const UStaticMeshCompo
 	return LightMapConverter->GetOrAdd(StaticMeshComponent);
 }
 
-FGLTFJsonHotspot* FGLTFConvertBuilder::AddUniqueHotspot(const AGLTFHotspotActor* HotspotActor)
-{
-	if (HotspotActor == nullptr)
-	{
-		return nullptr;
-	}
-
-	return HotspotConverter->GetOrAdd(HotspotActor);
-}
-
 FGLTFJsonSkySphere* FGLTFConvertBuilder::AddUniqueSkySphere(const AActor* SkySphereActor)
 {
 	if (SkySphereActor == nullptr)
diff --git a/Engine/Plugins/Enterprise/GLTFExporter/Source/GLTFExporter/Private/Builders/GLTFJsonBuilder.cpp b/Engine/Plugins/Enterprise/GLTFExporter/Source/GLTFExporter/Private/Builders/GLTFJsonBuilder.cpp
--- a/Engine/Plugins/Enterprise/GLTFExporter/Source/GLTFExporter/Private/Builders/GLTFJsonBuilder.cpp
+++ b/Engine/Plugins/Enterprise/GLTFExporter/Source/GLTFExporter/Private/Builders/GLTFJsonBuilder.cpp
@@ -111,11 +111,6 @@ FGLTFJsonBackdrop* FGLTFJsonBuilder::AddBackdrop()
 	return JsonRoot.Backdrops.Add();
 }
 
-FGLTFJsonHotspot* FGLTFJsonBuilder::AddHotspot()
-{
-	return JsonRoot.Hotspots.Add();
-}
-
 FGLTFJsonLight* FGLTFJsonBuilder::AddLight()
 {
 	return JsonRoot.Lights.Add();
diff --git a/Engine/Plugins/Enterprise/GLTFExporter/Source/GLTFExporter/Private/Converters/GLTFBlueprintUtility.cpp b/Engine/Plugins/Enterprise/GLTFExporter/Source/GLTFExporter/Private/Converters/GLTFBlueprintUtility.cpp
--- a/Engine/Plugins/Enterprise/GLTFExporter/Source/GLTFExporter/Private/Converters/GLTFBlueprintUtility.cpp
+++ b/Engine/Plugins/Enterprise/GLTFExporter/Source/GLTFExporter/Private/Converters/GLTFBlueprintUtility.cpp
@@ -1,7 +1,6 @@
 // Copyright Epic Games, Inc. All Rights Reserved.
 
 #include "Converters/GLTFBlueprintUtility.h"
-#include "Actors/GLTFHotspotActor.h"
 #include "LevelSequenceActor.h"
 
 FString FGLTFBlueprintUtility::GetClassPath(const AActor* Actor)
diff --git a/Engine/Plugins/Enterprise/GLTFExporter/Source/GLTFExporter/Private/Converters/GLTFCameraConverters.cpp b/Engine/Plugins/Enterprise/GLTFExporter/Source/GLTFExporter/Private/Converters/GLTFCameraConverters.cpp
--- a/Engine/Plugins/Enterprise/GLTFExporter/Source/GLTFExporter/Private/Converters/GLTFCameraConverters.cpp
+++ b/Engine/Plugins/Enterprise/GLTFExporter/Source/GLTFExporter/Private/Converters/GLTFCameraConverters.cpp
@@ -4,7 +4,6 @@
 #include "Builders/GLTFContainerBuilder.h"
 #include "Utilities/GLTFCoreUtilities.h"
 #include "Converters/GLTFNameUtility.h"
-#include "Actors/GLTFCameraActor.h"
 #include "Camera/CameraComponent.h"
 
 FGLTFJsonCamera* FGLTFCameraConverter::Convert(const UCameraComponent* CameraComponent)
@@ -50,44 +49,5 @@ FGLTFJsonCamera* FGLTFCameraConverter::Convert(const UCameraComponent* CameraCom
 			break;
 	}
 
-	const AActor* Owner = CameraComponent->GetOwner();
-	const AGLTFCameraActor* CameraActor = Owner != nullptr ? Cast<AGLTFCameraActor>(Owner) : nullptr;
-
-	if (CameraActor != nullptr)
-	{
-		if (Builder.ExportOptions->bExportCameraControls)
-		{
-			FGLTFJsonCameraControl CameraControl;
-			CameraControl.Mode = FGLTFCoreUtilities::ConvertCameraControlMode(CameraActor->Mode);
-			CameraControl.Target = Builder.AddUniqueNode(CameraActor->Target);
-			CameraControl.MaxDistance = FGLTFCoreUtilities::ConvertLength(CameraActor->DistanceMax, ExportScale);
-			CameraControl.MinDistance = FGLTFCoreUtilities::ConvertLength(CameraActor->DistanceMin, ExportScale);
-			CameraControl.MaxPitch = CameraActor->PitchAngleMax;
-			CameraControl.MinPitch = CameraActor->PitchAngleMin;
-
-			if (CameraActor->UsesYawLimits())
-			{
-				// Transform yaw limits to match right-handed system and glTF specification for cameras, i.e
-				// positive rotation is CCW, and camera looks down Z- (instead of X+).
-				const float MaxYaw = FMath::Max(-CameraActor->YawAngleMin, -CameraActor->YawAngleMax) - 90.0f;
-				const float MinYaw = FMath::Min(-CameraActor->YawAngleMin, -CameraActor->YawAngleMax) - 90.0f;
-
-				// We prefer the limits to be in the 0..360 range, but we only use MaxYaw to calculate
-				// the needed offset since we need to keep both limits a fixed distance apart from each other.
-				const float PositiveRangeOffset = FRotator::ClampAxis(MaxYaw) - MaxYaw;
-
-				CameraControl.MaxYaw = MaxYaw + PositiveRangeOffset;
-				CameraControl.MinYaw = MinYaw + PositiveRangeOffset;
-			}
-
-			CameraControl.RotationSensitivity = CameraActor->RotationSensitivity;
-			CameraControl.RotationInertia = CameraActor->RotationInertia;
-			CameraControl.DollySensitivity = CameraActor->DollySensitivity;
-			CameraControl.DollyDuration = CameraActor->DollyDuration;
-
-			JsonCamera->CameraControl = CameraControl;
-		}
-	}
-
 	return JsonCamera;
 }
diff --git a/Engine/Plugins/Enterprise/GLTFExporter/Source/GLTFExporter/Private/Converters/GLTFHotspotConverters.cpp b/Engine/Plugins/Enterprise/GLTFExporter/Source/GLTFExporter/Private/Converters/GLTFHotspotConverters.cpp
deleted file mode 100644
diff --git a/Engine/Plugins/Enterprise/GLTFExporter/Source/GLTFExporter/Private/Converters/GLTFNameUtility.h b/Engine/Plugins/Enterprise/GLTFExporter/Source/GLTFExporter/Private/Converters/GLTFNameUtility.h
--- a/Engine/Plugins/Enterprise/GLTFExporter/Source/GLTFExporter/Private/Converters/GLTFNameUtility.h
+++ b/Engine/Plugins/Enterprise/GLTFExporter/Source/GLTFExporter/Private/Converters/GLTFNameUtility.h
@@ -8,7 +8,6 @@
 
 class ULightComponent;
 class UCameraComponent;
-class AGLTFHotspotActor;
 
 class UCameraComponent;
 class ULightComponent;
diff --git a/Engine/Plugins/Enterprise/GLTFExporter/Source/GLTFExporter/Private/Converters/GLTFNodeConverters.cpp b/Engine/Plugins/Enterprise/GLTFExporter/Source/GLTFExporter/Private/Converters/GLTFNodeConverters.cpp
--- a/Engine/Plugins/Enterprise/GLTFExporter/Source/GLTFExporter/Private/Converters/GLTFNodeConverters.cpp
+++ b/Engine/Plugins/Enterprise/GLTFExporter/Source/GLTFExporter/Private/Converters/GLTFNodeConverters.cpp
@@ -5,7 +5,6 @@
 #include "Utilities/GLTFCoreUtilities.h"
 #include "Converters/GLTFBlueprintUtility.h"
 #include "Converters/GLTFNameUtility.h"
-#include "Actors/GLTFHotspotActor.h"
 #include "LevelSequenceActor.h"
 #include "Camera/CameraComponent.h"
 #include "Components/LightComponent.h"
@@ -53,13 +52,6 @@ FGLTFJsonNode* FGLTFActorConverter::Convert(const AActor* Actor)
 			Builder.AddUniqueAnimation(LevelSequenceActor);
 		}
 	}
-	else if (const AGLTFHotspotActor* HotspotActor = Cast<AGLTFHotspotActor>(Actor))
-	{
-		if (Builder.ExportOptions->bExportAnimationHotspots)
-		{
-			RootNode->Hotspot = Builder.AddUniqueHotspot(HotspotActor);
-		}
-	}
 	else
 	{
 		// TODO: add support for exporting brush geometry?
@@ -152,7 +144,7 @@ void FGLTFComponentConverter::ConvertComponentSpecialization(const USceneCompone
 		return;
 	}
 
-	if (Owner->IsA<ALevelSequenceActor>() || Owner->IsA<APawn>() || Owner->IsA<AGLTFHotspotActor>())
+	if (Owner->IsA<ALevelSequenceActor>() || Owner->IsA<APawn>())
 	{
 		return;
 	}
diff --git a/Engine/Plugins/Enterprise/GLTFExporter/Source/GLTFExporter/Private/Json/GLTFJsonCamera.cpp b/Engine/Plugins/Enterprise/GLTFExporter/Source/GLTFExporter/Private/Json/GLTFJsonCamera.cpp
--- a/Engine/Plugins/Enterprise/GLTFExporter/Source/GLTFExporter/Private/Json/GLTFJsonCamera.cpp
+++ b/Engine/Plugins/Enterprise/GLTFExporter/Source/GLTFExporter/Private/Json/GLTFJsonCamera.cpp
@@ -49,11 +49,4 @@ void FGLTFJsonCamera::WriteObject(IGLTFJsonWriter& Writer) const
 		default:
 			break;
 	}
-
-	if (CameraControl.IsSet())
-	{
-		Writer.StartExtensions();
-		Writer.Write(EGLTFJsonExtension::EPIC_CameraControls, CameraControl.GetValue());
-		Writer.EndExtensions();
-	}
 }
diff --git a/Engine/Plugins/Enterprise/GLTFExporter/Source/GLTFExporter/Private/Json/GLTFJsonCameraControl.cpp b/Engine/Plugins/Enterprise/GLTFExporter/Source/GLTFExporter/Private/Json/GLTFJsonCameraControl.cpp
deleted file mode 100644
diff --git a/Engine/Plugins/Enterprise/GLTFExporter/Source/GLTFExporter/Private/Json/GLTFJsonHotspot.cpp b/Engine/Plugins/Enterprise/GLTFExporter/Source/GLTFExporter/Private/Json/GLTFJsonHotspot.cpp
deleted file mode 100644
diff --git a/Engine/Plugins/Enterprise/GLTFExporter/Source/GLTFExporter/Private/Json/GLTFJsonNode.cpp b/Engine/Plugins/Enterprise/GLTFExporter/Source/GLTFExporter/Private/Json/GLTFJsonNode.cpp
--- a/Engine/Plugins/Enterprise/GLTFExporter/Source/GLTFExporter/Private/Json/GLTFJsonNode.cpp
+++ b/Engine/Plugins/Enterprise/GLTFExporter/Source/GLTFExporter/Private/Json/GLTFJsonNode.cpp
@@ -5,7 +5,6 @@
 #include "Json/GLTFJsonSkin.h"
 #include "Json/GLTFJsonMesh.h"
 #include "Json/GLTFJsonBackdrop.h"
-#include "Json/GLTFJsonHotspot.h"
 #include "Json/GLTFJsonLight.h"
 #include "Json/GLTFJsonLightMap.h"
 #include "Json/GLTFJsonSkySphere.h"
@@ -47,7 +46,7 @@ void FGLTFJsonNode::WriteObject(IGLTFJsonWriter& Writer) const
 		Writer.Write(TEXT("mesh"), Mesh);
 	}
 
-	if (Backdrop != nullptr || Hotspot != nullptr || Light != nullptr || LightMap != nullptr || SkySphere != nullptr)
+	if (Backdrop != nullptr || Light != nullptr || LightMap != nullptr || SkySphere != nullptr)
 	{
 		Writer.StartExtensions();
 
@@ -58,13 +57,6 @@ void FGLTFJsonNode::WriteObject(IGLTFJsonWriter& Writer) const
 			Writer.EndExtension();
 		}
 
-		if (Hotspot != nullptr)
-		{
-			Writer.StartExtension(EGLTFJsonExtension::EPIC_AnimationHotspots);
-			Writer.Write(TEXT("hotspot"), Hotspot);
-			Writer.EndExtension();
-		}
-
 		if (Light != nullptr)
 		{
 			Writer.StartExtension(EGLTFJsonExtension::KHR_LightsPunctual);
diff --git a/Engine/Plugins/Enterprise/GLTFExporter/Source/GLTFExporter/Private/Json/GLTFJsonRoot.cpp b/Engine/Plugins/Enterprise/GLTFExporter/Source/GLTFExporter/Private/Json/GLTFJsonRoot.cpp
--- a/Engine/Plugins/Enterprise/GLTFExporter/Source/GLTFExporter/Private/Json/GLTFJsonRoot.cpp
+++ b/Engine/Plugins/Enterprise/GLTFExporter/Source/GLTFExporter/Private/Json/GLTFJsonRoot.cpp
@@ -25,7 +25,7 @@ void FGLTFJsonRoot::WriteObject(IGLTFJsonWriter& Writer) const
 	if (Skins.Num() > 0) Writer.Write(TEXT("skins"), Skins);
 	if (Textures.Num() > 0) Writer.Write(TEXT("textures"), Textures);
 
-	if (Backdrops.Num() > 0 || Hotspots.Num() > 0 || Lights.Num() > 0 || LightMaps.Num() > 0 || SkySpheres.Num() > 0 || EpicLevelVariantSets.Num() > 0 || KhrMaterialVariants.Num() > 0)
+	if (Backdrops.Num() > 0 || Lights.Num() > 0 || LightMaps.Num() > 0 || SkySpheres.Num() > 0 || EpicLevelVariantSets.Num() > 0 || KhrMaterialVariants.Num() > 0)
 	{
 		Writer.StartExtensions();
 
@@ -36,13 +36,6 @@ void FGLTFJsonRoot::WriteObject(IGLTFJsonWriter& Writer) const
 			Writer.EndExtension();
 		}
 
-		if (Hotspots.Num() > 0)
-		{
-			Writer.StartExtension(EGLTFJsonExtension::EPIC_AnimationHotspots);
-			Writer.Write(TEXT("hotspots"), Hotspots);
-			Writer.EndExtension();
-		}
-
 		if (EpicLevelVariantSets.Num() > 0)
 		{
 			Writer.StartExtension(EGLTFJsonExtension::EPIC_LevelVariantSets);
diff --git a/Engine/Plugins/Enterprise/GLTFExporter/Source/GLTFExporter/Private/Options/GLTFExportOptions.cpp b/Engine/Plugins/Enterprise/GLTFExporter/Source/GLTFExporter/Private/Options/GLTFExportOptions.cpp
--- a/Engine/Plugins/Enterprise/GLTFExporter/Source/GLTFExporter/Private/Options/GLTFExportOptions.cpp
+++ b/Engine/Plugins/Enterprise/GLTFExporter/Source/GLTFExporter/Private/Options/GLTFExportOptions.cpp
@@ -40,8 +40,6 @@ void UGLTFExportOptions::ResetToDefault()
 	bExportHiddenInGame = false;
 	ExportLights = static_cast<int32>(EGLTFSceneMobility::Stationary | EGLTFSceneMobility::Movable);
 	bExportCameras = true;
-	bExportCameraControls = false;
-	bExportAnimationHotspots = false;
 	bExportHDRIBackdrops = false;
 	bExportSkySpheres = false;
 	VariantSetsMode = EGLTFVariantSetsMode::None;
diff --git a/Engine/Plugins/Enterprise/GLTFExporter/Source/GLTFExporter/Private/Utilities/GLTFCoreUtilities.cpp b/Engine/Plugins/Enterprise/GLTFExporter/Source/GLTFExporter/Private/Utilities/GLTFCoreUtilities.cpp
--- a/Engine/Plugins/Enterprise/GLTFExporter/Source/GLTFExporter/Private/Utilities/GLTFCoreUtilities.cpp
+++ b/Engine/Plugins/Enterprise/GLTFExporter/Source/GLTFExporter/Private/Utilities/GLTFCoreUtilities.cpp
@@ -2,7 +2,6 @@
 
 #include "Utilities/GLTFCoreUtilities.h"
 #include "Converters/GLTFTextureUtility.h"
-#include "Actors/GLTFCameraActor.h"
 
 float FGLTFCoreUtilities::ConvertLength(const float Length, const float ConversionScale)
 {
@@ -320,13 +319,3 @@ EGLTFJsonCubeFace FGLTFCoreUtilities::ConvertCubeFace(ECubeFace CubeFace)
 		default:            return EGLTFJsonCubeFace::None; // TODO: add error handling in callers
 	}
 }
-
-EGLTFJsonCameraControlMode FGLTFCoreUtilities::ConvertCameraControlMode(EGLTFCameraControlMode CameraMode)
-{
-	switch (CameraMode)
-	{
-		case EGLTFCameraControlMode::FreeLook: return EGLTFJsonCameraControlMode::FreeLook;
-		case EGLTFCameraControlMode::Orbital:  return EGLTFJsonCameraControlMode::Orbital;
-		default:                               return EGLTFJsonCameraControlMode::None; // TODO: add error handling in callers
-	}
-}
diff --git a/Engine/Plugins/Enterprise/GLTFExporter/Source/GLTFExporter/Private/Utilities/GLTFJsonUtilities.cpp b/Engine/Plugins/Enterprise/GLTFExporter/Source/GLTFExporter/Private/Utilities/GLTFJsonUtilities.cpp
--- a/Engine/Plugins/Enterprise/GLTFExporter/Source/GLTFExporter/Private/Utilities/GLTFJsonUtilities.cpp
+++ b/Engine/Plugins/Enterprise/GLTFExporter/Source/GLTFExporter/Private/Utilities/GLTFJsonUtilities.cpp
@@ -12,10 +12,8 @@ const TCHAR* FGLTFJsonUtilities::GetValue(EGLTFJsonExtension Enum)
 		case EGLTFJsonExtension::KHR_MaterialsVariants:   return TEXT("KHR_materials_variants");
 		case EGLTFJsonExtension::KHR_MeshQuantization:    return TEXT("KHR_mesh_quantization");
 		case EGLTFJsonExtension::KHR_TextureTransform:    return TEXT("KHR_texture_transform");
-		case EGLTFJsonExtension::EPIC_AnimationHotspots:  return TEXT("EPIC_animation_hotspots");
 		case EGLTFJsonExtension::EPIC_AnimationPlayback:  return TEXT("EPIC_animation_playback");
 		case EGLTFJsonExtension::EPIC_BlendModes:         return TEXT("EPIC_blend_modes");
-		case EGLTFJsonExtension::EPIC_CameraControls:     return TEXT("EPIC_camera_controls");
 		case EGLTFJsonExtension::EPIC_HDRIBackdrops:      return TEXT("EPIC_hdri_backdrops");
 		case EGLTFJsonExtension::EPIC_LevelVariantSets:   return TEXT("EPIC_level_variant_sets");
 		case EGLTFJsonExtension::EPIC_LightmapTextures:   return TEXT("EPIC_lightmap_textures");
@@ -162,18 +160,6 @@ const TCHAR* FGLTFJsonUtilities::GetValue(EGLTFJsonTargetPath Enum)
 	}
 }
 
-const TCHAR* FGLTFJsonUtilities::GetValue(EGLTFJsonCameraControlMode Enum)
-{
-	switch (Enum)
-	{
-		case EGLTFJsonCameraControlMode::FreeLook: return TEXT("freeLook");
-		case EGLTFJsonCameraControlMode::Orbital:  return TEXT("orbital");
-		default:
-			checkNoEntry();
-			return TEXT("");
-	}
-}
-
 const TCHAR* FGLTFJsonUtilities::GetValue(EGLTFJsonShadingModel Enum)
 {
 	switch (Enum)
diff --git a/Engine/Plugins/Enterprise/GLTFExporter/Source/GLTFExporter/Public/Actors/GLTFCameraActor.h b/Engine/Plugins/Enterprise/GLTFExporter/Source/GLTFExporter/Public/Actors/GLTFCameraActor.h
deleted file mode 100644
diff --git a/Engine/Plugins/Enterprise/GLTFExporter/Source/GLTFExporter/Public/Actors/GLTFHotspotActor.h b/Engine/Plugins/Enterprise/GLTFExporter/Source/GLTFExporter/Public/Actors/GLTFHotspotActor.h
deleted file mode 100644
diff --git a/Engine/Plugins/Enterprise/GLTFExporter/Source/GLTFExporter/Public/Builders/GLTFConvertBuilder.h b/Engine/Plugins/Enterprise/GLTFExporter/Source/GLTFExporter/Public/Builders/GLTFConvertBuilder.h
--- a/Engine/Plugins/Enterprise/GLTFExporter/Source/GLTFExporter/Public/Builders/GLTFConvertBuilder.h
+++ b/Engine/Plugins/Enterprise/GLTFExporter/Source/GLTFExporter/Public/Builders/GLTFConvertBuilder.h
@@ -19,7 +19,6 @@
 #include "Converters/GLTFBackdropConverters.h"
 #include "Converters/GLTFEpicLevelVariantSetsConverters.h"
 #include "Converters/GLTFLightMapConverters.h"
-#include "Converters/GLTFHotspotConverters.h"
 #include "Converters/GLTFKhrVariantConverters.h"
 #include "Converters/GLTFSkySphereConverters.h"
 
@@ -94,7 +93,6 @@ public:
 	FGLTFJsonLight* AddUniqueLight(const ULightComponent* LightComponent);
 	FGLTFJsonBackdrop* AddUniqueBackdrop(const AActor* BackdropActor);
 	FGLTFJsonLightMap* AddUniqueLightMap(const UStaticMeshComponent* StaticMeshComponent);
-	FGLTFJsonHotspot* AddUniqueHotspot(const AGLTFHotspotActor* HotspotActor);
 	FGLTFJsonSkySphere* AddUniqueSkySphere(const AActor* SkySphereActor);
 
 	FGLTFJsonEpicLevelVariantSets* AddUniqueEpicLevelVariantSets(const ULevelVariantSets* LevelVariantSets);
@@ -145,7 +143,6 @@ public:
 	TUniquePtr<IGLTFLightConverter> LightConverter = MakeUnique<FGLTFLightConverter>(*this);
 	TUniquePtr<IGLTFBackdropConverter> BackdropConverter = MakeUnique<FGLTFBackdropConverter>(*this);
 	TUniquePtr<IGLTFLightMapConverter> LightMapConverter = MakeUnique<FGLTFLightMapConverter>(*this);
-	TUniquePtr<IGLTFHotspotConverter> HotspotConverter = MakeUnique<FGLTFHotspotConverter>(*this);
 	TUniquePtr<IGLTFSkySphereConverter> SkySphereConverter = MakeUnique<FGLTFSkySphereConverter>(*this);
 
 	TUniquePtr<IGLTFEpicLevelVariantSetsConverter> EpicLevelVariantSetsConverter = MakeUnique<FGLTFEpicLevelVariantSetsConverter>(*this);
diff --git a/Engine/Plugins/Enterprise/GLTFExporter/Source/GLTFExporter/Public/Builders/GLTFJsonBuilder.h b/Engine/Plugins/Enterprise/GLTFExporter/Source/GLTFExporter/Public/Builders/GLTFJsonBuilder.h
--- a/Engine/Plugins/Enterprise/GLTFExporter/Source/GLTFExporter/Public/Builders/GLTFJsonBuilder.h
+++ b/Engine/Plugins/Enterprise/GLTFExporter/Source/GLTFExporter/Public/Builders/GLTFJsonBuilder.h
@@ -29,7 +29,6 @@ public:
 	FGLTFJsonSkin* AddSkin();
 	FGLTFJsonTexture* AddTexture();
 	FGLTFJsonBackdrop* AddBackdrop();
-	FGLTFJsonHotspot* AddHotspot();
 	FGLTFJsonLight* AddLight();
 	FGLTFJsonLightMap* AddLightMap();
 	FGLTFJsonSkySphere* AddSkySphere();
diff --git a/Engine/Plugins/Enterprise/GLTFExporter/Source/GLTFExporter/Public/Converters/GLTFHotspotConverters.h b/Engine/Plugins/Enterprise/GLTFExporter/Source/GLTFExporter/Public/Converters/GLTFHotspotConverters.h
deleted file mode 100644
diff --git a/Engine/Plugins/Enterprise/GLTFExporter/Source/GLTFExporter/Public/Json/GLTFJsonCamera.h b/Engine/Plugins/Enterprise/GLTFExporter/Source/GLTFExporter/Public/Json/GLTFJsonCamera.h
--- a/Engine/Plugins/Enterprise/GLTFExporter/Source/GLTFExporter/Public/Json/GLTFJsonCamera.h
+++ b/Engine/Plugins/Enterprise/GLTFExporter/Source/GLTFExporter/Public/Json/GLTFJsonCamera.h
@@ -3,7 +3,6 @@
 #pragma once
 
 #include "Json/GLTFJsonCore.h"
-#include "Json/GLTFJsonCameraControl.h"
 
 struct GLTFEXPORTER_API FGLTFJsonOrthographic : IGLTFJsonObject
 {
@@ -45,9 +44,7 @@ struct GLTFEXPORTER_API FGLTFJsonCamera : IGLTFJsonIndexedObject
 {
 	FString Name;
 
-	EGLTFJsonCameraType               Type;
-	TOptional<FGLTFJsonCameraControl> CameraControl;
-
+	EGLTFJsonCameraType   Type;
 	FGLTFJsonOrthographic Orthographic;
 	FGLTFJsonPerspective  Perspective;
 
diff --git a/Engine/Plugins/Enterprise/GLTFExporter/Source/GLTFExporter/Public/Json/GLTFJsonCameraControl.h b/Engine/Plugins/Enterprise/GLTFExporter/Source/GLTFExporter/Public/Json/GLTFJsonCameraControl.h
deleted file mode 100644
diff --git a/Engine/Plugins/Enterprise/GLTFExporter/Source/GLTFExporter/Public/Json/GLTFJsonCore.h b/Engine/Plugins/Enterprise/GLTFExporter/Source/GLTFExporter/Public/Json/GLTFJsonCore.h
--- a/Engine/Plugins/Enterprise/GLTFExporter/Source/GLTFExporter/Public/Json/GLTFJsonCore.h
+++ b/Engine/Plugins/Enterprise/GLTFExporter/Source/GLTFExporter/Public/Json/GLTFJsonCore.h
@@ -31,7 +31,6 @@ struct FGLTFJsonScene;
 struct FGLTFJsonSkin;
 struct FGLTFJsonTexture;
 struct FGLTFJsonBackdrop;
-struct FGLTFJsonHotspot;
 struct FGLTFJsonLight;
 struct FGLTFJsonLightMap;
 struct FGLTFJsonSkySphere;
diff --git a/Engine/Plugins/Enterprise/GLTFExporter/Source/GLTFExporter/Public/Json/GLTFJsonEnums.h b/Engine/Plugins/Enterprise/GLTFExporter/Source/GLTFExporter/Public/Json/GLTFJsonEnums.h
--- a/Engine/Plugins/Enterprise/GLTFExporter/Source/GLTFExporter/Public/Json/GLTFJsonEnums.h
+++ b/Engine/Plugins/Enterprise/GLTFExporter/Source/GLTFExporter/Public/Json/GLTFJsonEnums.h
@@ -10,10 +10,8 @@ enum class EGLTFJsonExtension
 	KHR_MaterialsVariants,
 	KHR_MeshQuantization,
 	KHR_TextureTransform,
-	EPIC_AnimationHotspots,
 	EPIC_AnimationPlayback,
 	EPIC_BlendModes,
-	EPIC_CameraControls,
 	EPIC_HDRIBackdrops,
 	EPIC_LevelVariantSets,
 	EPIC_LightmapTextures,
@@ -164,10 +162,3 @@ enum class EGLTFJsonLightType
 	Point,
 	Spot
 };
-
-enum class EGLTFJsonCameraControlMode
-{
-	None = -1,
-	FreeLook,
-	Orbital
-};
diff --git a/Engine/Plugins/Enterprise/GLTFExporter/Source/GLTFExporter/Public/Json/GLTFJsonHotspot.h b/Engine/Plugins/Enterprise/GLTFExporter/Source/GLTFExporter/Public/Json/GLTFJsonHotspot.h
deleted file mode 100644
diff --git a/Engine/Plugins/Enterprise/GLTFExporter/Source/GLTFExporter/Public/Json/GLTFJsonNode.h b/Engine/Plugins/Enterprise/GLTFExporter/Source/GLTFExporter/Public/Json/GLTFJsonNode.h
--- a/Engine/Plugins/Enterprise/GLTFExporter/Source/GLTFExporter/Public/Json/GLTFJsonNode.h
+++ b/Engine/Plugins/Enterprise/GLTFExporter/Source/GLTFExporter/Public/Json/GLTFJsonNode.h
@@ -16,7 +16,6 @@ struct GLTFEXPORTER_API FGLTFJsonNode : IGLTFJsonIndexedObject
 	FGLTFJsonSkin*      Skin;
 	FGLTFJsonMesh*      Mesh;
 	FGLTFJsonBackdrop*  Backdrop;
-	FGLTFJsonHotspot*   Hotspot;
 	FGLTFJsonLight*     Light;
 	FGLTFJsonLightMap*  LightMap;
 	FGLTFJsonSkySphere* SkySphere;
@@ -38,7 +37,6 @@ protected:
 		, Skin(nullptr)
 		, Mesh(nullptr)
 		, Backdrop(nullptr)
-		, Hotspot(nullptr)
 		, Light(nullptr)
 		, LightMap(nullptr)
 		, SkySphere(nullptr)
diff --git a/Engine/Plugins/Enterprise/GLTFExporter/Source/GLTFExporter/Public/Json/GLTFJsonRoot.h b/Engine/Plugins/Enterprise/GLTFExporter/Source/GLTFExporter/Public/Json/GLTFJsonRoot.h
--- a/Engine/Plugins/Enterprise/GLTFExporter/Source/GLTFExporter/Public/Json/GLTFJsonRoot.h
+++ b/Engine/Plugins/Enterprise/GLTFExporter/Source/GLTFExporter/Public/Json/GLTFJsonRoot.h
@@ -18,7 +18,6 @@
 #include "Json/GLTFJsonSkin.h"
 #include "Json/GLTFJsonTexture.h"
 #include "Json/GLTFJsonBackdrop.h"
-#include "Json/GLTFJsonHotspot.h"
 #include "Json/GLTFJsonLight.h"
 #include "Json/GLTFJsonLightMap.h"
 #include "Json/GLTFJsonSkySphere.h"
@@ -47,7 +46,6 @@ struct GLTFEXPORTER_API FGLTFJsonRoot : IGLTFJsonObject
 	TGLTFJsonIndexedObjectArray<FGLTFJsonSkin>       Skins;
 	TGLTFJsonIndexedObjectArray<FGLTFJsonTexture>    Textures;
 	TGLTFJsonIndexedObjectArray<FGLTFJsonBackdrop>   Backdrops;
-	TGLTFJsonIndexedObjectArray<FGLTFJsonHotspot>    Hotspots;
 	TGLTFJsonIndexedObjectArray<FGLTFJsonLight>      Lights;
 	TGLTFJsonIndexedObjectArray<FGLTFJsonLightMap>   LightMaps;
 	TGLTFJsonIndexedObjectArray<FGLTFJsonSkySphere>  SkySpheres;
diff --git a/Engine/Plugins/Enterprise/GLTFExporter/Source/GLTFExporter/Public/Options/GLTFExportOptions.h b/Engine/Plugins/Enterprise/GLTFExporter/Source/GLTFExporter/Public/Options/GLTFExportOptions.h
--- a/Engine/Plugins/Enterprise/GLTFExporter/Source/GLTFExporter/Public/Options/GLTFExportOptions.h
+++ b/Engine/Plugins/Enterprise/GLTFExporter/Source/GLTFExporter/Public/Options/GLTFExportOptions.h
@@ -155,7 +155,7 @@ class GLTFEXPORTER_API UGLTFExportOptions : public UObject
 	UPROPERTY(EditAnywhere, BlueprintReadWrite, Config, Category = Animation)
 	bool bExportLevelSequences;
 
-	/** If enabled, export single animation asset used by a skeletal mesh component or hotspot actor. Export of vertex skin weights must be enabled. */
+	/** If enabled, export single animation asset used by a skeletal mesh component. Export of vertex skin weights must be enabled. */
 	UPROPERTY(EditAnywhere, BlueprintReadWrite, Config, Category = Animation, Meta = (EditCondition = "bExportVertexSkinWeights"))
 	bool bExportAnimationSequences;
 
@@ -203,14 +203,6 @@ class GLTFEXPORTER_API UGLTFExportOptions : public UObject
 	UPROPERTY(EditAnywhere, BlueprintReadWrite, Config, Category = Scene)
 	bool bExportCameras;
 
-	/** If enabled, export GLTFCameraActors. Uses extension EPIC_camera_controls, which is supported by Unreal's glTF viewer. */
-	UPROPERTY(EditAnywhere, BlueprintReadWrite, Config, Category = Scene, Meta = (EditCondition = "bExportCameras"))
-	bool bExportCameraControls;
-
-	/** If enabled, export GLTFHotspotActors. Uses extension EPIC_animation_hotspots, which is supported by Unreal's glTF viewer. */
-	UPROPERTY(EditAnywhere, BlueprintReadWrite, Config, Category = Scene)
-	bool bExportAnimationHotspots;
-
 	/** If enabled, export HDRIBackdrop blueprints. Uses extension EPIC_hdri_backdrops, which is supported by Unreal's glTF viewer. */
 	UPROPERTY(EditAnywhere, BlueprintReadWrite, Config, Category = Scene, Meta = (DisplayName = "Export HDRI Backdrops"))
 	bool bExportHDRIBackdrops;
diff --git a/Engine/Plugins/Enterprise/GLTFExporter/Source/GLTFExporter/Public/Utilities/GLTFCoreUtilities.h b/Engine/Plugins/Enterprise/GLTFExporter/Source/GLTFExporter/Public/Utilities/GLTFCoreUtilities.h
--- a/Engine/Plugins/Enterprise/GLTFExporter/Source/GLTFExporter/Public/Utilities/GLTFCoreUtilities.h
+++ b/Engine/Plugins/Enterprise/GLTFExporter/Source/GLTFExporter/Public/Utilities/GLTFCoreUtilities.h
@@ -15,8 +15,6 @@
 #include "RHIDefinitions.h"
 #include "SceneTypes.h"
 
-enum class EGLTFCameraControlMode : unsigned char;
-
 struct GLTFEXPORTER_API FGLTFCoreUtilities
 {
 	static float ConvertLength(const float Length, const float ConversionScale = 0.01);
@@ -76,8 +74,6 @@ struct GLTFEXPORTER_API FGLTFCoreUtilities
 
 	static EGLTFJsonCubeFace ConvertCubeFace(ECubeFace CubeFace);
 
-	static EGLTFJsonCameraControlMode ConvertCameraControlMode(EGLTFCameraControlMode CameraMode);
-
 	template <typename ComponentType>
 	static EGLTFJsonComponentType GetComponentType()
 	{
diff --git a/Engine/Plugins/Enterprise/GLTFExporter/Source/GLTFExporter/Public/Utilities/GLTFJsonUtilities.h b/Engine/Plugins/Enterprise/GLTFExporter/Source/GLTFExporter/Public/Utilities/GLTFJsonUtilities.h
--- a/Engine/Plugins/Enterprise/GLTFExporter/Source/GLTFExporter/Public/Utilities/GLTFJsonUtilities.h
+++ b/Engine/Plugins/Enterprise/GLTFExporter/Source/GLTFExporter/Public/Utilities/GLTFJsonUtilities.h
@@ -24,6 +24,5 @@ struct GLTFEXPORTER_API FGLTFJsonUtilities
 	static const TCHAR* GetValue(EGLTFJsonLightType Enum);
 	static const TCHAR* GetValue(EGLTFJsonInterpolation Enum);
 	static const TCHAR* GetValue(EGLTFJsonTargetPath Enum);
-	static const TCHAR* GetValue(EGLTFJsonCameraControlMode Enum);
 	static const TCHAR* GetValue(EGLTFJsonShadingModel Enum);
 };
