Fix horde issue 218449

diff --git a/Engine/Plugins/Enterprise/GLTFExporter/Source/GLTFExporter/Private/Materials/GLTFProxyMaterialParameterInfo.h b/Engine/Plugins/Enterprise/GLTFExporter/Source/GLTFExporter/Private/Materials/GLTFProxyMaterialParameterInfo.h
--- a/Engine/Plugins/Enterprise/GLTFExporter/Source/GLTFExporter/Private/Materials/GLTFProxyMaterialParameterInfo.h
+++ b/Engine/Plugins/Enterprise/GLTFExporter/Source/GLTFExporter/Private/Materials/GLTFProxyMaterialParameterInfo.h
@@ -5,7 +5,6 @@
 #include "Utilities/GLTFProxyMaterialUtilities.h"
 
 class UMaterialInstanceConstant;
-class UMaterialInstanceDynamic;
 class UMaterialInterface;
 class UTexture;
 
@@ -24,11 +23,6 @@ public:
 		return FGLTFProxyMaterialUtilities::GetParameterValue(Material, ParameterInfo, OutValue, NonDefaultOnly);
 	}
 
-	void Set(UMaterialInstanceDynamic* Material, const ParameterType& Value, bool NonDefaultOnly = false) const
-	{
-		FGLTFProxyMaterialUtilities::SetParameterValue(Material, ParameterInfo, Value, NonDefaultOnly);
-	}
-
 #if WITH_EDITOR
 	void Set(UMaterialInstanceConstant* Material, const ParameterType& Value, bool NonDefaultOnly = false) const
 	{
