Enable GLTFExporter plugin by default and bump version to 1.3.1

diff --git a/Engine/Plugins/Enterprise/GLTFExporter/GLTFExporter.uplugin b/Engine/Plugins/Enterprise/GLTFExporter/GLTFExporter.uplugin
--- a/Engine/Plugins/Enterprise/GLTFExporter/GLTFExporter.uplugin
+++ b/Engine/Plugins/Enterprise/GLTFExporter/GLTFExporter.uplugin
@@ -1,15 +1,15 @@
 {
 	"FileVersion": 3,
-	"Version": 130,
-	"VersionName": "1.3.0",
+	"Version": 131,
+	"VersionName": "1.3.1",
 	"FriendlyName": "glTF Exporter",
 	"Description": "An exporter for Khronos glTF 2.0.",
 	"Category": "Exporters",
 	"CreatedBy": "Epic Games, Inc.",
 	"CreatedByURL": "http://epicgames.com",
 	"DocsURL": "",
-	"MarketplaceURL": "",
 	"SupportURL": "",
+	"EnabledByDefault": true,
 	"CanContainContent": true,
 	"Modules": [
 		{
