Add API docs to UGLTFExporter::ExportToGLTF and default to current playing UWorld if no Object param provided to export.

diff --git a/Engine/Plugins/Enterprise/GLTFExporter/Source/GLTFExporter/Private/Exporters/GLTFExporter.cpp b/Engine/Plugins/Enterprise/GLTFExporter/Source/GLTFExporter/Private/Exporters/GLTFExporter.cpp
--- a/Engine/Plugins/Enterprise/GLTFExporter/Source/GLTFExporter/Private/Exporters/GLTFExporter.cpp
+++ b/Engine/Plugins/Enterprise/GLTFExporter/Source/GLTFExporter/Private/Exporters/GLTFExporter.cpp
@@ -75,8 +75,19 @@ bool UGLTFExporter::ExportToGLTF(UObject* Object, const FString& FilePath, const
 {
 	if (Object == nullptr)
 	{
-		OutMessages.Errors.Add(TEXT("No object to export"));
-		return false;
+		Object = GEngine->GetCurrentPlayWorld();
+#if WITH_EDITOR
+		if (Object == nullptr)
+		{
+			Object = GEditor->GetEditorWorldContext().World();
+		}
+#endif
+
+		if (Object == nullptr)
+		{
+			OutMessages.Errors.Add(TEXT("No object to export"));
+			return false;
+		}
 	}
 
 	UGLTFExporter* Exporter = Cast<UGLTFExporter>(FindExporter(Object, TEXT("gltf")));
diff --git a/Engine/Plugins/Enterprise/GLTFExporter/Source/GLTFExporter/Public/Exporters/GLTFExporter.h b/Engine/Plugins/Enterprise/GLTFExporter/Source/GLTFExporter/Public/Exporters/GLTFExporter.h
--- a/Engine/Plugins/Enterprise/GLTFExporter/Source/GLTFExporter/Public/Exporters/GLTFExporter.h
+++ b/Engine/Plugins/Enterprise/GLTFExporter/Source/GLTFExporter/Public/Exporters/GLTFExporter.h
@@ -35,8 +35,30 @@ public:
 
 	virtual bool ExportBinary(UObject* Object, const TCHAR* Type, FArchive& Archive, FFeedbackContext* Warn, int32 FileIndex, uint32 PortFlags) override;
 
+	/**
+	 * Export the specified object to a glTF file (.gltf or .glb)
+	 *
+	 * @param Object          The object to export (supported types are UMaterialInterface, UStaticMesh, USkeletalMesh, UWorld, UAnimSequence, ULevelSequence, ULevelVariantSets). Will default to currently active world if null.
+	 * @param FilePath        The filename on disk to save as. Associated textures and binary files will be saved in the same folder, unless file extension is .glb - which results in a self-contained binary file.
+	 * @param Options         The various options to use during export. Will default to the project's user-specific editor settings if null.
+	 * @param SelectedActors  The set of actors to export, only applicable if Object is a UWorld. An empty set results in export of all actors.
+	 * @param OutMessages     The resulting log messages from the export.
+	 *
+	 * @return	true if the the object was successfully exported
+	 */
 	UFUNCTION(BlueprintCallable, Category = "Miscellaneous")
 	static bool ExportToGLTF(UObject* Object, const FString& FilePath, const UGLTFExportOptions* Options, const TSet<AActor*>& SelectedActors, FGLTFExportMessages& OutMessages);
+
+	/**
+	 * Export the specified object to a glTF file (.gltf or .glb)
+	 *
+	 * @param Object          The object to export (supported types are UMaterialInterface, UStaticMesh, USkeletalMesh, UWorld, UAnimSequence, ULevelSequence, ULevelVariantSets). Will default to currently active world if null.
+	 * @param FilePath        The filename on disk to save as. Associated textures and binary files will be saved in the same folder, unless file extension is .glb - which results in a self-contained binary file.
+	 * @param Options         The various options to use during export. Will default to the project's user-specific editor settings if null.
+	 * @param SelectedActors  The set of actors to export, only applicable if Object is a UWorld. An empty set results in export of all actors.
+	 *
+	 * @return	true if the the object was successfully exported
+	 */
 	static bool ExportToGLTF(UObject* Object, const FString& FilePath, const UGLTFExportOptions* Options = nullptr, const TSet<AActor*>& SelectedActors = {});
 
 protected:
