Fix glTF export of camera far clipping plane to correctly use infinite projection

diff --git a/Engine/Plugins/Enterprise/GLTFExporter/Source/GLTFExporter/Private/Converters/GLTFCameraConverters.cpp b/Engine/Plugins/Enterprise/GLTFExporter/Source/GLTFExporter/Private/Converters/GLTFCameraConverters.cpp
--- a/Engine/Plugins/Enterprise/GLTFExporter/Source/GLTFExporter/Private/Converters/GLTFCameraConverters.cpp
+++ b/Engine/Plugins/Enterprise/GLTFExporter/Source/GLTFExporter/Private/Converters/GLTFCameraConverters.cpp
@@ -35,9 +35,7 @@ FGLTFJsonCamera* FGLTFCameraConverter::Convert(const UCameraComponent* CameraCom
 				JsonCamera->Perspective.AspectRatio = DesiredView.AspectRatio;
 			}
 			JsonCamera->Perspective.YFov = FGLTFCoreUtilities::ConvertFieldOfView(DesiredView.FOV, DesiredView.AspectRatio);
-			// NOTE: even thought ZFar is optional, if we don't set it, then most gltf viewers won't handle it well.
-			JsonCamera->Perspective.ZFar = FGLTFCoreUtilities::ConvertLength(WORLD_MAX, ExportScale); // TODO: Unreal doesn't have max draw distance per view?
-			JsonCamera->Perspective.ZNear = FGLTFCoreUtilities::ConvertLength(GNearClippingPlane, ExportScale);
+			JsonCamera->Perspective.ZNear = FGLTFCoreUtilities::ConvertLength(DesiredView.GetFinalPerspectiveNearClipPlane(), ExportScale);
 			break;
 
 		case EGLTFJsonCameraType::None:
