Fix  error LNK2005 in GLTFMaterialBakingModule.cpp regarding FMaterialBakingModuleImpl

diff --git a/Engine/Plugins/Enterprise/GLTFExporter/Source/GLTFMaterialBaking/Private/GLTFMaterialBakingModule.cpp b/Engine/Plugins/Enterprise/GLTFExporter/Source/GLTFMaterialBaking/Private/GLTFMaterialBakingModule.cpp
--- a/Engine/Plugins/Enterprise/GLTFExporter/Source/GLTFMaterialBaking/Private/GLTFMaterialBakingModule.cpp
+++ b/Engine/Plugins/Enterprise/GLTFExporter/Source/GLTFMaterialBaking/Private/GLTFMaterialBakingModule.cpp
@@ -55,7 +55,7 @@ static TAutoConsoleVariable<int32> CVarMaterialBakingRDOCCapture(
 	TEXT("1: Turned On"),
 	ECVF_Default);
 
-namespace FMaterialBakingModuleImpl
+namespace FGLTFMaterialBakingModuleImpl
 {
 	// Custom dynamic mesh allocator specifically tailored for Material Baking.
 	// This will always reuse the same couple buffers, so searching linearly is not a problem.
@@ -391,7 +391,7 @@ void FGLTFMaterialBakingModule::BakeMaterials(const TArray<FGLTFMaterialDataEx*>
 	const int32 NumMaterials = MaterialSettings.Num();
 	const bool bSaveIntermediateTextures = CVarSaveIntermediateTextures.GetValueOnAnyThread() == 1;
 
-	using namespace FMaterialBakingModuleImpl;
+	using namespace FGLTFMaterialBakingModuleImpl;
 	FMaterialBakingDynamicMeshBufferAllocator MaterialBakingDynamicMeshBufferAllocator;
 
 	FScopedSlowTask Progress(NumMaterials, LOCTEXT("BakeMaterials", "Baking Materials..."), true );
@@ -434,7 +434,7 @@ void FGLTFMaterialBakingModule::BakeMaterials(const TArray<FGLTFMaterialDataEx*>
 	{
 		TRACE_CPUPROFILER_EVENT_SCOPE(PrepareRenderItems);
 
-		TMap<FMaterialBakingModuleImpl::FRenderItemKey, FGLTFMeshMaterialRenderItem*>* RenderItems = new TMap<FRenderItemKey, FGLTFMeshMaterialRenderItem *>();
+		TMap<FGLTFMaterialBakingModuleImpl::FRenderItemKey, FGLTFMeshMaterialRenderItem*>* RenderItems = new TMap<FRenderItemKey, FGLTFMeshMaterialRenderItem *>();
 		const FGLTFMaterialDataEx* CurrentMaterialSettings = MaterialSettings[MaterialIndex];
 		const FGLTFMeshRenderData* CurrentMeshSettings = MeshSettings[MaterialIndex];
 
