Fix minor typo in API documentation of UGLTFExporter::ExportToGLTF

diff --git a/Engine/Plugins/Enterprise/GLTFExporter/Source/GLTFExporter/Public/Exporters/GLTFExporter.h b/Engine/Plugins/Enterprise/GLTFExporter/Source/GLTFExporter/Public/Exporters/GLTFExporter.h
--- a/Engine/Plugins/Enterprise/GLTFExporter/Source/GLTFExporter/Public/Exporters/GLTFExporter.h
+++ b/Engine/Plugins/Enterprise/GLTFExporter/Source/GLTFExporter/Public/Exporters/GLTFExporter.h
@@ -38,13 +38,13 @@ public:
 	/**
 	 * Export the specified object to a glTF file (.gltf or .glb)
 	 *
-	 * @param Object          The object to export (supported types are UMaterialInterface, UStaticMesh, USkeletalMesh, UWorld, UAnimSequence, ULevelSequence, ULevelVariantSets). Will default to currently active world if null.
+	 * @param Object          The object to export (supported types are UMaterialInterface, UStaticMesh, USkeletalMesh, UWorld, UAnimSequence, ULevelSequence, ULevelVariantSets). Will default to the currently active world if null.
 	 * @param FilePath        The filename on disk to save as. Associated textures and binary files will be saved in the same folder, unless file extension is .glb - which results in a self-contained binary file.
 	 * @param Options         The various options to use during export. Will default to the project's user-specific editor settings if null.
-	 * @param SelectedActors  The set of actors to export, only applicable if Object is a UWorld. An empty set results in export of all actors.
+	 * @param SelectedActors  The set of actors to export, only applicable if the object to export is a UWorld. An empty set results in the export of all actors.
 	 * @param OutMessages     The resulting log messages from the export.
 	 *
-	 * @return	true if the the object was successfully exported
+	 * @return true if the object was successfully exported
 	 */
 	UFUNCTION(BlueprintCallable, Category = "Miscellaneous")
 	static bool ExportToGLTF(UObject* Object, const FString& FilePath, const UGLTFExportOptions* Options, const TSet<AActor*>& SelectedActors, FGLTFExportMessages& OutMessages);
@@ -52,12 +52,12 @@ public:
 	/**
 	 * Export the specified object to a glTF file (.gltf or .glb)
 	 *
-	 * @param Object          The object to export (supported types are UMaterialInterface, UStaticMesh, USkeletalMesh, UWorld, UAnimSequence, ULevelSequence, ULevelVariantSets). Will default to currently active world if null.
+	 * @param Object          The object to export (supported types are UMaterialInterface, UStaticMesh, USkeletalMesh, UWorld, UAnimSequence, ULevelSequence, ULevelVariantSets). Will default to the currently active world if null.
 	 * @param FilePath        The filename on disk to save as. Associated textures and binary files will be saved in the same folder, unless file extension is .glb - which results in a self-contained binary file.
 	 * @param Options         The various options to use during export. Will default to the project's user-specific editor settings if null.
-	 * @param SelectedActors  The set of actors to export, only applicable if Object is a UWorld. An empty set results in export of all actors.
+	 * @param SelectedActors  The set of actors to export, only applicable if the object to export is a UWorld. An empty set results in the export of all actors.
 	 *
-	 * @return	true if the the object was successfully exported
+	 * @return true if the object was successfully exported
 	 */
 	static bool ExportToGLTF(UObject* Object, const FString& FilePath, const UGLTFExportOptions* Options = nullptr, const TSet<AActor*>& SelectedActors = {});
 
