Fix horde issue 218606

diff --git a/Engine/Plugins/Enterprise/GLTFExporter/Source/GLTFExporter/Private/Exporters/GLTFExporter.cpp b/Engine/Plugins/Enterprise/GLTFExporter/Source/GLTFExporter/Private/Exporters/GLTFExporter.cpp
--- a/Engine/Plugins/Enterprise/GLTFExporter/Source/GLTFExporter/Private/Exporters/GLTFExporter.cpp
+++ b/Engine/Plugins/Enterprise/GLTFExporter/Source/GLTFExporter/Private/Exporters/GLTFExporter.cpp
@@ -8,6 +8,9 @@
 #include "Builders/GLTFContainerBuilder.h"
 #include "UObject/GCObjectScopeGuard.h"
 #include "AssetExportTask.h"
+#if WITH_EDITOR
+#include "Editor.h"
+#endif
 
 UGLTFExporter::UGLTFExporter(const FObjectInitializer& ObjectInitializer)
 	: Super(ObjectInitializer)
