Fix BitmaskEnum warnings regarding UGLTFExportOptions::NoLossyImageFormatFor and UGLTFExportOptions::ExportLights

diff --git a/Engine/Plugins/Enterprise/GLTFExporter/Source/GLTFExporter/Public/Options/GLTFExportOptions.h b/Engine/Plugins/Enterprise/GLTFExporter/Source/GLTFExporter/Public/Options/GLTFExportOptions.h
--- a/Engine/Plugins/Enterprise/GLTFExporter/Source/GLTFExporter/Public/Options/GLTFExportOptions.h
+++ b/Engine/Plugins/Enterprise/GLTFExporter/Source/GLTFExporter/Public/Options/GLTFExportOptions.h
@@ -172,7 +172,7 @@ class GLTFEXPORTER_API UGLTFExportOptions : public UObject
 	int32 TextureImageQuality;
 
 	/** Texture types that will always use lossless formats (e.g. PNG) because of sensitivity to compression artifacts. */
-	UPROPERTY(EditAnywhere, BlueprintReadWrite, Config, Category = Texture, Meta = (Bitmask, BitmaskEnum = EGLTFTextureType, EditCondition = "TextureImageFormat == EGLTFTextureImageFormat::JPEG"))
+	UPROPERTY(EditAnywhere, BlueprintReadWrite, Config, Category = Texture, Meta = (Bitmask, BitmaskEnum = "/Script/GLTFExporter.EGLTFTextureType", EditCondition = "TextureImageFormat == EGLTFTextureImageFormat::JPEG"))
 	int32 NoLossyImageFormatFor; // Bitmask combined from EGLTFTextureType
 
 	/** If enabled, export UV tiling and un-mirroring settings in a texture coordinate expression node for simple material input expressions. Uses extension KHR_texture_transform. */
@@ -196,7 +196,7 @@ class GLTFEXPORTER_API UGLTFExportOptions : public UObject
 	bool bExportHiddenInGame;
 
 	/** Mobility of directional, point, and spot light components that will be exported. Uses extension KHR_lights_punctual. */
-	UPROPERTY(EditAnywhere, BlueprintReadWrite, Config, Category = Scene, Meta = (Bitmask, BitmaskEnum = EGLTFSceneMobility))
+	UPROPERTY(EditAnywhere, BlueprintReadWrite, Config, Category = Scene, Meta = (Bitmask, BitmaskEnum = "/Script/GLTFExporter.EGLTFSceneMobility"))
 	int32 ExportLights; // Bitmask combined from EGLTFSceneMobility
 
 	/** If enabled, export camera components. */
