Don't let FGLTFProxyAssetActions filter types : it's a proxy for actions that are already filtered and therefore shows duplicate filters for proxied types (Material, Material Instance, ...)

#rnx

diff --git a/Engine/Plugins/Enterprise/GLTFExporter/Source/GLTFExporter/Private/Actions/GLTFProxyAssetActions.h b/Engine/Plugins/Enterprise/GLTFExporter/Source/GLTFExporter/Private/Actions/GLTFProxyAssetActions.h
--- a/Engine/Plugins/Enterprise/GLTFExporter/Source/GLTFExporter/Private/Actions/GLTFProxyAssetActions.h
+++ b/Engine/Plugins/Enterprise/GLTFExporter/Source/GLTFExporter/Private/Actions/GLTFProxyAssetActions.h
@@ -24,7 +24,7 @@ public:
 	virtual bool CanRename(const FAssetData& InAsset, FText* OutErrorMsg) const override { return OriginalActions->CanRename(InAsset, OutErrorMsg); }
 	virtual bool CanDuplicate(const FAssetData& InAsset, FText* OutErrorMsg) const override { return OriginalActions->CanDuplicate(InAsset, OutErrorMsg); }
 	virtual TArray<FAssetData> GetValidAssetsForPreviewOrEdit(TArrayView<const FAssetData> InAssetDatas, bool bIsPreview) override { return OriginalActions->GetValidAssetsForPreviewOrEdit(InAssetDatas, bIsPreview); }
-	virtual bool CanFilter() override { return OriginalActions->CanFilter(); }
+	virtual bool CanFilter() override { return false; }
 	virtual FName GetFilterName() const override { return OriginalActions->GetFilterName(); }
 	virtual bool CanLocalize() const override { return OriginalActions->CanLocalize(); }
 	virtual bool CanMerge() const override { return OriginalActions->CanMerge(); }
