Fix cache misses when using DDC pakfile due to bad TransmittanceColor custom attrib registration

diff --git a/Engine/Plugins/Enterprise/GLTFExporter/Source/GLTFMaterialBaking/Private/AddCustomAttributeHack.cpp b/Engine/Plugins/Enterprise/GLTFExporter/Source/GLTFMaterialBaking/Private/AddCustomAttributeHack.cpp
deleted file mode 100644
diff --git a/Engine/Plugins/Enterprise/GLTFExporter/Source/GLTFMaterialBaking/Private/AddCustomAttributeHack.h b/Engine/Plugins/Enterprise/GLTFExporter/Source/GLTFMaterialBaking/Private/AddCustomAttributeHack.h
deleted file mode 100644
diff --git a/Engine/Plugins/Enterprise/GLTFExporter/Source/GLTFMaterialBaking/Private/GLTFMaterialBakingModule.cpp b/Engine/Plugins/Enterprise/GLTFExporter/Source/GLTFMaterialBaking/Private/GLTFMaterialBakingModule.cpp
--- a/Engine/Plugins/Enterprise/GLTFExporter/Source/GLTFMaterialBaking/Private/GLTFMaterialBakingModule.cpp
+++ b/Engine/Plugins/Enterprise/GLTFExporter/Source/GLTFMaterialBaking/Private/GLTFMaterialBakingModule.cpp
@@ -7,7 +7,6 @@
 #include "UObject/UObjectGlobals.h"
 #include "GLTFMaterialBakingStructures.h"
 #include "GLTFMaterialBakingHelpers.h"
-#include "AddCustomAttributeHack.h"
 #include "Async/Async.h"
 #include "Async/ParallelFor.h"
 #include "Materials/MaterialInstance.h"
@@ -296,9 +295,6 @@ void FGLTFMaterialBakingModule::StartupModule()
 
 	// Register callback on garbage collection
 	FCoreUObjectDelegates::GetPreGarbageCollectDelegate().AddRaw(this, &FGLTFMaterialBakingModule::OnPreGarbageCollect);
-
-	// NOTE: Because the attribute definition of TransmittanceColor is not registered by the engine, we have to do it here
-	AddCustomAttributeHack(FGuid(0xF2D8C70E, 0x42ECA0D1, 0x4652D0AD, 0xB785A065), "TransmittanceColor", "GetThinTranslucentMaterialOutput", MCT_Float3, FVector4(0.5, 0.5, 0.5, 0));
 }
 
 void FGLTFMaterialBakingModule::ShutdownModule()
