Fix for errors in FGLTFProxyAssetActions.h caused by changes to IAssetTypeActions.h in 22816987 which added two more functions needing to be overriden.

#rnx

diff --git a/Engine/Plugins/Enterprise/GLTFExporter/Source/GLTFExporter/Private/Actions/GLTFProxyAssetActions.h b/Engine/Plugins/Enterprise/GLTFExporter/Source/GLTFExporter/Private/Actions/GLTFProxyAssetActions.h
--- a/Engine/Plugins/Enterprise/GLTFExporter/Source/GLTFExporter/Private/Actions/GLTFProxyAssetActions.h
+++ b/Engine/Plugins/Enterprise/GLTFExporter/Source/GLTFExporter/Private/Actions/GLTFProxyAssetActions.h
@@ -19,6 +19,8 @@ public:
 	virtual FColor GetTypeColor() const override { return OriginalActions->GetTypeColor(); }
 	virtual void GetActions(const TArray<UObject*>& InObjects, FMenuBuilder& MenuBuilder) override { OriginalActions->GetActions(InObjects, MenuBuilder); }
 	virtual void OpenAssetEditor(const TArray<UObject*>& InObjects, TSharedPtr<IToolkitHost> EditWithinLevelEditor) override { OriginalActions->OpenAssetEditor(InObjects, EditWithinLevelEditor); }
+	virtual void OpenAssetEditor(const TArray<UObject*>& InObjects, EAssetTypeActivationOpenedMethod OpenedMethod, TSharedPtr<IToolkitHost> EditWithinLevelEditor = TSharedPtr<IToolkitHost>()) override { OriginalActions->OpenAssetEditor(InObjects, OpenedMethod, EditWithinLevelEditor); }
+	virtual bool SupportsOpenedMethod(EAssetTypeActivationOpenedMethod OpenedMethod) const { return OriginalActions->SupportsOpenedMethod(OpenedMethod); }
 	virtual bool AssetsActivatedOverride(const TArray<UObject*>& InObjects, EAssetTypeActivationMethod::Type ActivationType) override { return OriginalActions->AssetsActivatedOverride(InObjects, ActivationType); }
 	virtual bool CanRename(const FAssetData& InAsset, FText* OutErrorMsg) const override { return OriginalActions->CanRename(InAsset, OutErrorMsg); }
 	virtual bool CanDuplicate(const FAssetData& InAsset, FText* OutErrorMsg) const override { return OriginalActions->CanDuplicate(InAsset, OutErrorMsg); }
