Fix crash when trying to export rect lights to gltf

diff --git a/Engine/Plugins/Enterprise/GLTFExporter/Source/GLTFExporter/Private/Converters/GLTFLightConverters.cpp b/Engine/Plugins/Enterprise/GLTFExporter/Source/GLTFExporter/Private/Converters/GLTFLightConverters.cpp
--- a/Engine/Plugins/Enterprise/GLTFExporter/Source/GLTFExporter/Private/Converters/GLTFLightConverters.cpp
+++ b/Engine/Plugins/Enterprise/GLTFExporter/Source/GLTFExporter/Private/Converters/GLTFLightConverters.cpp
@@ -10,17 +10,17 @@
 
 FGLTFJsonLight* FGLTFLightConverter::Convert(const ULightComponent* LightComponent)
 {
-	FGLTFJsonLight* Light = Builder.AddLight();
-
-	Light->Name = FGLTFNameUtility::GetName(LightComponent);
-	Light->Type = FGLTFCoreUtilities::ConvertLightType(LightComponent->GetLightType());
-
-	if (Light->Type == EGLTFJsonLightType::None)
+	const EGLTFJsonLightType Type = FGLTFCoreUtilities::ConvertLightType(LightComponent->GetLightType());
+	if (Type == EGLTFJsonLightType::None)
 	{
 		// TODO: report error (unsupported light component type)
 		return nullptr;
 	}
 
+	FGLTFJsonLight* Light = Builder.AddLight();
+	Light->Name = FGLTFNameUtility::GetName(LightComponent);
+	Light->Type = Type;
+
 	Light->Intensity = LightComponent->Intensity;
 
 	const FLinearColor TemperatureColor = LightComponent->bUseTemperature ? FLinearColor::MakeFromColorTemperature(LightComponent->Temperature) : FLinearColor::White;
