Avoid creating incomplete json objects during gltf export

diff --git a/Engine/Plugins/Enterprise/GLTFExporter/Source/GLTFExporter/Private/Converters/GLTFCameraConverters.cpp b/Engine/Plugins/Enterprise/GLTFExporter/Source/GLTFExporter/Private/Converters/GLTFCameraConverters.cpp
--- a/Engine/Plugins/Enterprise/GLTFExporter/Source/GLTFExporter/Private/Converters/GLTFCameraConverters.cpp
+++ b/Engine/Plugins/Enterprise/GLTFExporter/Source/GLTFExporter/Private/Converters/GLTFCameraConverters.cpp
@@ -2,16 +2,23 @@
 
 #include "Converters/GLTFCameraConverters.h"
 #include "Builders/GLTFContainerBuilder.h"
-#include "GameFramework/Actor.h"
 #include "Utilities/GLTFCoreUtilities.h"
 #include "Converters/GLTFNameUtility.h"
 #include "Camera/CameraComponent.h"
+#include "GameFramework/Actor.h"
 
 FGLTFJsonCamera* FGLTFCameraConverter::Convert(const UCameraComponent* CameraComponent)
 {
+	const EGLTFJsonCameraType Type = FGLTFCoreUtilities::ConvertCameraType(CameraComponent->ProjectionMode);
+	if (Type == EGLTFJsonCameraType::None)
+	{
+		// TODO: report error (unsupported camera type)
+		return nullptr;
+	}
+
 	FGLTFJsonCamera* JsonCamera = Builder.AddCamera();
 	JsonCamera->Name = FGLTFNameUtility::GetName(CameraComponent);
-	JsonCamera->Type = FGLTFCoreUtilities::ConvertCameraType(CameraComponent->ProjectionMode);
+	JsonCamera->Type = Type;
 
 	FMinimalViewInfo DesiredView;
 	const_cast<UCameraComponent*>(CameraComponent)->GetCameraView(0, DesiredView);
@@ -39,10 +46,6 @@ FGLTFJsonCamera* FGLTFCameraConverter::Convert(const UCameraComponent* CameraCom
 			JsonCamera->Perspective.ZNear = FGLTFCoreUtilities::ConvertLength(DesiredView.GetFinalPerspectiveNearClipPlane(), ExportScale);
 			break;
 
-		case EGLTFJsonCameraType::None:
-			// TODO: report error (unsupported camera type)
-			return nullptr;
-
 		default:
 			checkNoEntry();
 			break;
diff --git a/Engine/Plugins/Enterprise/GLTFExporter/Source/GLTFExporter/Private/Converters/GLTFEpicLevelVariantSetsConverters.cpp b/Engine/Plugins/Enterprise/GLTFExporter/Source/GLTFExporter/Private/Converters/GLTFEpicLevelVariantSetsConverters.cpp
--- a/Engine/Plugins/Enterprise/GLTFExporter/Source/GLTFExporter/Private/Converters/GLTFEpicLevelVariantSetsConverters.cpp
+++ b/Engine/Plugins/Enterprise/GLTFExporter/Source/GLTFExporter/Private/Converters/GLTFEpicLevelVariantSetsConverters.cpp
@@ -15,8 +15,7 @@
 
 FGLTFJsonEpicLevelVariantSets* FGLTFEpicLevelVariantSetsConverter::Convert(const ULevelVariantSets* LevelVariantSets)
 {
-	FGLTFJsonEpicLevelVariantSets* JsonLevelVariantSets = Builder.AddEpicLevelVariantSets();
-	LevelVariantSets->GetName(JsonLevelVariantSets->Name);
+	TArray<FGLTFJsonEpicVariantSet> VariantSets;
 
 	for (const UVariantSet* VariantSet: LevelVariantSets->GetVariantSets())
 	{
@@ -34,15 +33,18 @@ FGLTFJsonEpicLevelVariantSets* FGLTFEpicLevelVariantSetsConverter::Convert(const
 
 		if (JsonVariantSet.Variants.Num() > 0)
 		{
-			JsonLevelVariantSets->VariantSets.Add(JsonVariantSet);
+			VariantSets.Add(JsonVariantSet);
 		}
 	}
 
-	if (JsonLevelVariantSets->VariantSets.Num() == 0)
+	if (VariantSets.Num() == 0)
 	{
 		return nullptr;
 	}
 
+	FGLTFJsonEpicLevelVariantSets* JsonLevelVariantSets = Builder.AddEpicLevelVariantSets();
+	LevelVariantSets->GetName(JsonLevelVariantSets->Name);
+	JsonLevelVariantSets->VariantSets = VariantSets;
 	return JsonLevelVariantSets;
 }
 
diff --git a/Engine/Plugins/Enterprise/GLTFExporter/Source/GLTFExporter/Private/Converters/GLTFKhrVariantConverters.cpp b/Engine/Plugins/Enterprise/GLTFExporter/Source/GLTFExporter/Private/Converters/GLTFKhrVariantConverters.cpp
--- a/Engine/Plugins/Enterprise/GLTFExporter/Source/GLTFExporter/Private/Converters/GLTFKhrVariantConverters.cpp
+++ b/Engine/Plugins/Enterprise/GLTFExporter/Source/GLTFExporter/Private/Converters/GLTFKhrVariantConverters.cpp
@@ -16,16 +16,8 @@ FGLTFJsonKhrMaterialVariant* FGLTFKhrMaterialVariantConverter::Convert(const UVa
 		return nullptr;
 	}
 
-
-	FGLTFJsonKhrMaterialVariant* MaterialVariant = Builder.AddKhrMaterialVariant();
-
-	// TODO: add warning if the variant name is not unique, i.e it's already used?
-	// While material variants are technically allowed to use the same name, it may
-	// cause confusion when trying to select the correct variant in a viewer.
-	MaterialVariant->Name = Variant->GetDisplayText().ToString();
-
-	typedef TTuple<FGLTFJsonPrimitive*, FGLTFJsonMaterial*> TPrimitiveMaterial;
-	TArray<TPrimitiveMaterial> PrimitiveMaterials;
+	typedef TTuple<FGLTFJsonPrimitive*, FGLTFJsonMaterial*> FGLTFPrimitiveMaterial;
+	TArray<TTuple<FGLTFJsonPrimitive*, FGLTFJsonMaterial*>> PrimitiveMaterials;
 
 	for (const UVariantObjectBinding* Binding: Variant->GetBindings())
 	{
@@ -56,7 +48,13 @@ FGLTFJsonKhrMaterialVariant* FGLTFKhrMaterialVariantConverter::Convert(const UVa
 		return nullptr;
 	}
 
-	for (const TPrimitiveMaterial& PrimitiveMaterial: PrimitiveMaterials)
+	FGLTFJsonKhrMaterialVariant* MaterialVariant = Builder.AddKhrMaterialVariant();
+	// TODO: add warning if the variant name is not unique, i.e it's already used?
+	// While material variants are technically allowed to use the same name, it may
+	// cause confusion when trying to select the correct variant in a viewer.
+	MaterialVariant->Name = Variant->GetDisplayText().ToString();
+
+	for (const FGLTFPrimitiveMaterial& PrimitiveMaterial: PrimitiveMaterials)
 	{
 		FGLTFJsonPrimitive* Primitive = PrimitiveMaterial.Key;
 		FGLTFJsonMaterial* Material = PrimitiveMaterial.Value;
diff --git a/Engine/Plugins/Enterprise/GLTFExporter/Source/GLTFExporter/Private/Converters/GLTFSkinConverters.cpp b/Engine/Plugins/Enterprise/GLTFExporter/Source/GLTFExporter/Private/Converters/GLTFSkinConverters.cpp
--- a/Engine/Plugins/Enterprise/GLTFExporter/Source/GLTFExporter/Private/Converters/GLTFSkinConverters.cpp
+++ b/Engine/Plugins/Enterprise/GLTFExporter/Source/GLTFExporter/Private/Converters/GLTFSkinConverters.cpp
@@ -4,9 +4,7 @@
 #include "Utilities/GLTFCoreUtilities.h"
 #include "Converters/GLTFBoneUtility.h"
 #include "Builders/GLTFConvertBuilder.h"
-
 #include "Animation/Skeleton.h"
-#include "Components/SkeletalMeshComponent.h"
 #include "Engine/SkeletalMesh.h"
 
 FGLTFJsonSkin* FGLTFSkinConverter::Convert(FGLTFJsonNode* RootNode, const USkeletalMesh* SkeletalMesh)
@@ -14,10 +12,6 @@ FGLTFJsonSkin* FGLTFSkinConverter::Convert(FGLTFJsonNode* RootNode, const USkele
 	const USkeleton* Skeleton = SkeletalMesh->GetSkeleton();
 	const FReferenceSkeleton& RefSkeleton = SkeletalMesh->GetRefSkeleton();
 
-	FGLTFJsonSkin* JsonSkin = Builder.AddSkin();
-	JsonSkin->Name = Skeleton != nullptr ? Skeleton->GetName() : SkeletalMesh->GetName();
-	JsonSkin->Skeleton = RootNode;
-
 	const int32 BoneCount = RefSkeleton.GetNum();
 	if (BoneCount == 0)
 	{
@@ -25,6 +19,9 @@ FGLTFJsonSkin* FGLTFSkinConverter::Convert(FGLTFJsonNode* RootNode, const USkele
 		return nullptr;
 	}
 
+	FGLTFJsonSkin* JsonSkin = Builder.AddSkin();
+	JsonSkin->Name = Skeleton != nullptr ? Skeleton->GetName() : SkeletalMesh->GetName();
+	JsonSkin->Skeleton = RootNode;
 	JsonSkin->Joints.AddUninitialized(BoneCount);
 
 	for (int32 BoneIndex = 0; BoneIndex < BoneCount; ++BoneIndex)
diff --git a/Engine/Plugins/Enterprise/GLTFExporter/Source/GLTFExporter/Private/Converters/GLTFTextureConverters.cpp b/Engine/Plugins/Enterprise/GLTFExporter/Source/GLTFExporter/Private/Converters/GLTFTextureConverters.cpp
--- a/Engine/Plugins/Enterprise/GLTFExporter/Source/GLTFExporter/Private/Converters/GLTFTextureConverters.cpp
+++ b/Engine/Plugins/Enterprise/GLTFExporter/Source/GLTFExporter/Private/Converters/GLTFTextureConverters.cpp
@@ -18,14 +18,14 @@ void FGLTFTexture2DConverter::Sanitize(const UTexture2D*& Texture2D, bool& bToSR
 
 FGLTFJsonTexture* FGLTFTexture2DConverter::Convert(const UTexture2D* Texture2D, bool bToSRGB)
 {
-	if (Builder.ExportOptions->TextureImageFormat != EGLTFTextureImageFormat::None)
+	if (Builder.ExportOptions->TextureImageFormat == EGLTFTextureImageFormat::None)
 	{
-		FGLTFJsonTexture* JsonTexture = Builder.AddTexture();
-		Builder.ScheduleSlowTask<FGLTFDelayedTexture2DTask>(Builder, Texture2D, bToSRGB, JsonTexture);
-		return JsonTexture;
+		return nullptr;
 	}
 
-	return nullptr;
+	FGLTFJsonTexture* JsonTexture = Builder.AddTexture();
+	Builder.ScheduleSlowTask<FGLTFDelayedTexture2DTask>(Builder, Texture2D, bToSRGB, JsonTexture);
+	return JsonTexture;
 }
 
 void FGLTFTextureCubeConverter::Sanitize(const UTextureCube*& TextureCube, ECubeFace& CubeFace, bool& bToSRGB)
@@ -38,14 +38,14 @@ void FGLTFTextureCubeConverter::Sanitize(const UTextureCube*& TextureCube, ECube
 
 FGLTFJsonTexture* FGLTFTextureCubeConverter::Convert(const UTextureCube* TextureCube, ECubeFace CubeFace, bool bToSRGB)
 {
-	if (Builder.ExportOptions->TextureImageFormat != EGLTFTextureImageFormat::None)
+	if (Builder.ExportOptions->TextureImageFormat == EGLTFTextureImageFormat::None)
 	{
-		FGLTFJsonTexture* JsonTexture = Builder.AddTexture();
-		Builder.ScheduleSlowTask<FGLTFDelayedTextureCubeTask>(Builder, TextureCube, CubeFace, bToSRGB, JsonTexture);
-		return JsonTexture;
+		return nullptr;
 	}
 
-	return nullptr;
+	FGLTFJsonTexture* JsonTexture = Builder.AddTexture();
+	Builder.ScheduleSlowTask<FGLTFDelayedTextureCubeTask>(Builder, TextureCube, CubeFace, bToSRGB, JsonTexture);
+	return JsonTexture;
 }
 
 void FGLTFTextureRenderTarget2DConverter::Sanitize(const UTextureRenderTarget2D*& RenderTarget2D, bool& bToSRGB)
@@ -58,14 +58,14 @@ void FGLTFTextureRenderTarget2DConverter::Sanitize(const UTextureRenderTarget2D*
 
 FGLTFJsonTexture* FGLTFTextureRenderTarget2DConverter::Convert(const UTextureRenderTarget2D* RenderTarget2D, bool bToSRGB)
 {
-	if (Builder.ExportOptions->TextureImageFormat != EGLTFTextureImageFormat::None)
+	if (Builder.ExportOptions->TextureImageFormat == EGLTFTextureImageFormat::None)
 	{
-		FGLTFJsonTexture* JsonTexture = Builder.AddTexture();
-		Builder.ScheduleSlowTask<FGLTFDelayedTextureRenderTarget2DTask>(Builder, RenderTarget2D, bToSRGB, JsonTexture);
-		return JsonTexture;
+		return nullptr;
 	}
 
-	return nullptr;
+	FGLTFJsonTexture* JsonTexture = Builder.AddTexture();
+	Builder.ScheduleSlowTask<FGLTFDelayedTextureRenderTarget2DTask>(Builder, RenderTarget2D, bToSRGB, JsonTexture);
+	return JsonTexture;
 }
 
 void FGLTFTextureRenderTargetCubeConverter::Sanitize(const UTextureRenderTargetCube*& RenderTargetCube, ECubeFace& CubeFace, bool& bToSRGB)
@@ -78,26 +78,28 @@ void FGLTFTextureRenderTargetCubeConverter::Sanitize(const UTextureRenderTargetC
 
 FGLTFJsonTexture* FGLTFTextureRenderTargetCubeConverter::Convert(const UTextureRenderTargetCube* RenderTargetCube, ECubeFace CubeFace, bool bToSRGB)
 {
-	if (Builder.ExportOptions->TextureImageFormat != EGLTFTextureImageFormat::None)
+	if (Builder.ExportOptions->TextureImageFormat == EGLTFTextureImageFormat::None)
 	{
-		FGLTFJsonTexture* JsonTexture = Builder.AddTexture();
-		Builder.ScheduleSlowTask<FGLTFDelayedTextureRenderTargetCubeTask>(Builder, RenderTargetCube, CubeFace, bToSRGB, JsonTexture);
-		return JsonTexture;
+		return nullptr;
 	}
 
-	return nullptr;
+	FGLTFJsonTexture* JsonTexture = Builder.AddTexture();
+	Builder.ScheduleSlowTask<FGLTFDelayedTextureRenderTargetCubeTask>(Builder, RenderTargetCube, CubeFace, bToSRGB, JsonTexture);
+	return JsonTexture;
 }
 
 FGLTFJsonTexture* FGLTFTextureLightMapConverter::Convert(const ULightMapTexture2D* LightMap)
 {
-#if WITH_EDITOR
-	if (Builder.ExportOptions->TextureImageFormat != EGLTFTextureImageFormat::None)
+	if (Builder.ExportOptions->TextureImageFormat == EGLTFTextureImageFormat::None)
 	{
-		FGLTFJsonTexture* JsonTexture = Builder.AddTexture();
-		Builder.ScheduleSlowTask<FGLTFDelayedTextureLightMapTask>(Builder, LightMap, JsonTexture);
-		return JsonTexture;
+		return nullptr;
 	}
-#endif
 
+#if WITH_EDITOR
+	FGLTFJsonTexture* JsonTexture = Builder.AddTexture();
+	Builder.ScheduleSlowTask<FGLTFDelayedTextureLightMapTask>(Builder, LightMap, JsonTexture);
+	return JsonTexture;
+#else
 	return nullptr;
+#endif
 }
