Reimplement asset actions for creating glTF proxy materials using new proper engine api which reduces code and avoids many potential future issues

diff --git a/Engine/Plugins/Enterprise/GLTFExporter/Source/GLTFExporter/GLTFExporter.Build.cs b/Engine/Plugins/Enterprise/GLTFExporter/Source/GLTFExporter/GLTFExporter.Build.cs
--- a/Engine/Plugins/Enterprise/GLTFExporter/Source/GLTFExporter/GLTFExporter.Build.cs
+++ b/Engine/Plugins/Enterprise/GLTFExporter/Source/GLTFExporter/GLTFExporter.Build.cs
@@ -55,6 +55,7 @@ public class GLTFExporter : ModuleRules
 					"EditorStyle",
 					"PropertyEditor",
 					"ToolMenus",
+					"ContentBrowser",
 					"MaterialUtilities",
 					"MeshMergeUtilities",
 					"MeshDescription",
diff --git a/Engine/Plugins/Enterprise/GLTFExporter/Source/GLTFExporter/Private/Actions/GLTFProxyAssetActions.cpp b/Engine/Plugins/Enterprise/GLTFExporter/Source/GLTFExporter/Private/Actions/GLTFProxyAssetActions.cpp
--- a/Engine/Plugins/Enterprise/GLTFExporter/Source/GLTFExporter/Private/Actions/GLTFProxyAssetActions.cpp
+++ b/Engine/Plugins/Enterprise/GLTFExporter/Source/GLTFExporter/Private/Actions/GLTFProxyAssetActions.cpp
@@ -8,55 +8,68 @@
 #include "Options/GLTFProxyOptions.h"
 #include "UI/GLTFProxyOptionsWindow.h"
 #include "Converters/GLTFMaterialProxyFactory.h"
+#include "ContentBrowserMenuContexts.h"
 #include "ToolMenus.h"
 
 #define LOCTEXT_NAMESPACE "GLTFProxyAssetActions"
 
-FGLTFProxyAssetActions::FGLTFProxyAssetActions(const TSharedRef<IAssetTypeActions>& OriginalActions)
-	: OriginalActions(OriginalActions)
+FName FGLTFProxyAssetActions::MenuName = TEXT("ContentBrowser.AssetContextMenu.MaterialInterface");
+FName FGLTFProxyAssetActions::SectionName = TEXT("GetAssetActions");
+FName FGLTFProxyAssetActions::EntryName = TEXT("MenuEntry_CreateProxy");
+
+void FGLTFProxyAssetActions::AddMenuEntry()
 {
+	UToolMenus* ToolMenus = UToolMenus::TryGet();
+	if (ToolMenus == nullptr)
+	{
+		return;
+	}
+
+	UToolMenu* Menu = ToolMenus->ExtendMenu(MenuName);
+	FToolMenuSection& Section = Menu->FindOrAddSection(SectionName);
+	Section.AddDynamicEntry(EntryName, FNewToolMenuSectionDelegate::CreateStatic(&FGLTFProxyAssetActions::OnConstructMenu));
 }
 
-void FGLTFProxyAssetActions::GetActions(const TArray<UObject*>& InObjects, FToolMenuSection& Section)
+void FGLTFProxyAssetActions::RemoveMenuEntry()
 {
-	OriginalActions->GetActions(InObjects, Section);
-	GetProxyActions(InObjects, Section);
+	UToolMenus* ToolMenus = UToolMenus::TryGet();
+	if (ToolMenus == nullptr)
+	{
+		return;
+	}
+
+	ToolMenus->RemoveEntry(MenuName, SectionName, EntryName);
 }
 
-void FGLTFProxyAssetActions::GetProxyActions(const TArray<UObject*>& InObjects, FToolMenuSection& Section)
+void FGLTFProxyAssetActions::OnConstructMenu(FToolMenuSection& MenuSection)
 {
-	const TArray<FWeakObjectPtr> Objects(InObjects);
-
-	Section.AddMenuEntry(
-		"MenuEntry_CreateProxy",
-		LOCTEXT("MenuEntry_CreateProxy", "Create glTF Proxy Material"),
-		LOCTEXT("MenuEntry_CreateProxyTooltip", "Creates a proxy version of this material for glTF export."),
-		FSlateIcon(FGLTFEditorStyle::Get().GetStyleSetName(), "Icon16"),
-		FUIAction(FExecuteAction::CreateSP(this, &FGLTFProxyAssetActions::OnCreateProxy, Objects))
-		);
+	const TAttribute<FText> Label = LOCTEXT("MenuEntry_CreateProxy", "Create glTF Proxy Material");
+	const TAttribute<FText> ToolTip = LOCTEXT("MenuEntry_CreateProxyTooltip", "Creates a proxy version of this material for glTF export.");
+	const FSlateIcon Icon = FSlateIcon(FGLTFEditorStyle::Get().GetStyleSetName(), TEXT("Icon16"));
+	const FToolMenuExecuteAction Action = FToolMenuExecuteAction::CreateStatic(&FGLTFProxyAssetActions::OnExecuteAction);
+	MenuSection.AddMenuEntry("MenuEntry_CreateProxy", Label, ToolTip, Icon, Action);
 }
 
-void FGLTFProxyAssetActions::OnCreateProxy(TArray<FWeakObjectPtr> Objects) const
+void FGLTFProxyAssetActions::OnExecuteAction(const FToolMenuContext& MenuContext)
 {
-	UGLTFProxyOptions* Options = NewObject<UGLTFProxyOptions>();
+	const UContentBrowserAssetContextMenuContext* Context = UContentBrowserAssetContextMenuContext::FindContextWithAssets(MenuContext);
+	if (Context == nullptr)
+	{
+		return;
+	}
 
+	UGLTFProxyOptions* Options = NewObject<UGLTFProxyOptions>();
 	if (!SGLTFProxyOptionsWindow::ShowDialog(Options))
 	{
 		return;
 	}
 
 	FGLTFMaterialProxyFactory ProxyFactory(Options);
-
-	for (const FWeakObjectPtr& Object : Objects)
+	for (UMaterialInterface* Material : Context->LoadSelectedObjects<UMaterialInterface>())
 	{
-		if (UMaterialInterface* Material = Cast<UMaterialInterface>(Object.Get()))
-		{
-			ProxyFactory.RootPath = FPaths::GetPath(Material->GetPathName()) / TEXT("GLTF");
-			ProxyFactory.Create(Material);
-		}
+		ProxyFactory.RootPath = FPaths::GetPath(Material->GetPathName()) / TEXT("GLTF");
+		ProxyFactory.Create(Material);
 	}
-
-	ProxyFactory.OpenLog();
 }
 
 #undef LOCTEXT_NAMESPACE
diff --git a/Engine/Plugins/Enterprise/GLTFExporter/Source/GLTFExporter/Private/Actions/GLTFProxyAssetActions.h b/Engine/Plugins/Enterprise/GLTFExporter/Source/GLTFExporter/Private/Actions/GLTFProxyAssetActions.h
--- a/Engine/Plugins/Enterprise/GLTFExporter/Source/GLTFExporter/Private/Actions/GLTFProxyAssetActions.h
+++ b/Engine/Plugins/Enterprise/GLTFExporter/Source/GLTFExporter/Private/Actions/GLTFProxyAssetActions.h
@@ -2,61 +2,28 @@
 
 #pragma once
 
+#include "CoreMinimal.h"
+
 #if WITH_EDITOR
 
-#include "AssetTypeActions_Base.h"
+struct FToolMenuSection;
+struct FToolMenuContext;
 
-class FGLTFProxyAssetActions final : public IAssetTypeActions
+class FGLTFProxyAssetActions
 {
 public:
 
-	FGLTFProxyAssetActions(const TSharedRef<IAssetTypeActions>& OriginalActions);
-	
-	virtual void GetActions(const TArray<UObject*>& InObjects, FToolMenuSection& Section) override;
-
-	virtual FText GetName() const override { return OriginalActions->GetName(); }
-	virtual UClass* GetSupportedClass() const override { return OriginalActions->GetSupportedClass(); }
-	virtual FColor GetTypeColor() const override { return OriginalActions->GetTypeColor(); }
-	virtual void GetActions(const TArray<UObject*>& InObjects, FMenuBuilder& MenuBuilder) override { OriginalActions->GetActions(InObjects, MenuBuilder); }
-	virtual void OpenAssetEditor(const TArray<UObject*>& InObjects, TSharedPtr<IToolkitHost> EditWithinLevelEditor) override { OriginalActions->OpenAssetEditor(InObjects, EditWithinLevelEditor); }
-	virtual void OpenAssetEditor(const TArray<UObject*>& InObjects, EAssetTypeActivationOpenedMethod OpenedMethod, TSharedPtr<IToolkitHost> EditWithinLevelEditor = TSharedPtr<IToolkitHost>()) override { OriginalActions->OpenAssetEditor(InObjects, OpenedMethod, EditWithinLevelEditor); }
-	virtual bool SupportsOpenedMethod(EAssetTypeActivationOpenedMethod OpenedMethod) const override { return OriginalActions->SupportsOpenedMethod(OpenedMethod); }
-	virtual bool AssetsActivatedOverride(const TArray<UObject*>& InObjects, EAssetTypeActivationMethod::Type ActivationType) override { return OriginalActions->AssetsActivatedOverride(InObjects, ActivationType); }
-	virtual bool CanRename(const FAssetData& InAsset, FText* OutErrorMsg) const override { return OriginalActions->CanRename(InAsset, OutErrorMsg); }
-	virtual bool CanDuplicate(const FAssetData& InAsset, FText* OutErrorMsg) const override { return OriginalActions->CanDuplicate(InAsset, OutErrorMsg); }
-	virtual TArray<FAssetData> GetValidAssetsForPreviewOrEdit(TArrayView<const FAssetData> InAssetDatas, bool bIsPreview) override { return OriginalActions->GetValidAssetsForPreviewOrEdit(InAssetDatas, bIsPreview); }
-	virtual bool CanFilter() override { return false; }
-	virtual FName GetFilterName() const override { return OriginalActions->GetFilterName(); }
-	virtual bool CanLocalize() const override { return OriginalActions->CanLocalize(); }
-	virtual bool CanMerge() const override { return OriginalActions->CanMerge(); }
-	virtual void Merge( UObject* InObject) override { OriginalActions->Merge(InObject); }
-	virtual void Merge(UObject* BaseAsset, UObject* RemoteAsset, UObject* LocalAsset, const FOnMergeResolved& ResolutionCallback) override { OriginalActions->Merge(BaseAsset, RemoteAsset, LocalAsset, ResolutionCallback); }
-	virtual uint32 GetCategories() override { return OriginalActions->GetCategories(); }
-	virtual FString GetObjectDisplayName(UObject* Object) const override { return OriginalActions->GetObjectDisplayName(Object); }
-	virtual const TArray<FText>& GetSubMenus() const override { return OriginalActions->GetSubMenus(); }
-	virtual bool ShouldForceWorldCentric() override { return OriginalActions->ShouldForceWorldCentric(); }
-	virtual void PerformAssetDiff(UObject* OldAsset, UObject* NewAsset, const FRevisionInfo& OldRevision, const FRevisionInfo& NewRevision) const override { OriginalActions->PerformAssetDiff(OldAsset, NewAsset, OldRevision, NewRevision); }
-	virtual UThumbnailInfo* GetThumbnailInfo(UObject* Asset) const override { return OriginalActions->GetThumbnailInfo(Asset); }
-	virtual EThumbnailPrimType GetDefaultThumbnailPrimitiveType(UObject* Asset) const override { return OriginalActions->GetDefaultThumbnailPrimitiveType(Asset); }
-	virtual TSharedPtr<SWidget> GetThumbnailOverlay(const FAssetData& AssetData) const override { return OriginalActions->GetThumbnailOverlay(AssetData); }
-	virtual FText GetAssetDescription(const FAssetData& AssetData) const override { return OriginalActions->GetAssetDescription(AssetData); }
-	virtual bool IsImportedAsset() const override { return OriginalActions->IsImportedAsset(); }
-	virtual void GetResolvedSourceFilePaths(const TArray<UObject*>& TypeAssets, TArray<FString>& OutSourceFilePaths) const override { OriginalActions->GetResolvedSourceFilePaths(TypeAssets, OutSourceFilePaths); }
-	virtual void GetSourceFileLabels(const TArray<UObject*>& TypeAssets, TArray<FString>& OutSourceFileLabels) const override { OriginalActions->GetSourceFileLabels(TypeAssets, OutSourceFileLabels); }
-	virtual void BuildBackendFilter(FARFilter& InFilter) override { OriginalActions->BuildBackendFilter(InFilter); }
-	virtual FText GetDisplayNameFromAssetData(const FAssetData& AssetData) const override { return OriginalActions->GetDisplayNameFromAssetData(AssetData); }
-	virtual void SetSupported(bool bInSupported) override { OriginalActions->SetSupported(bInSupported); }
-	virtual bool IsSupported() const override { return OriginalActions->IsSupported(); }
-	virtual FTopLevelAssetPath GetClassPathName() const override { return OriginalActions->GetClassPathName(); };
-	virtual const FSlateBrush* GetIconBrush(const FAssetData& InAssetData, const FName InClassName) const override { return OriginalActions->GetIconBrush(InAssetData, InClassName); }
-	virtual const FSlateBrush* GetThumbnailBrush(const FAssetData& InAssetData, const FName InClassName) const override { return OriginalActions->GetThumbnailBrush(InAssetData, InClassName); }
-private:
+	static FName MenuName;
+	static FName SectionName;
+	static FName EntryName;
 
-	void GetProxyActions(const TArray<UObject*>& InObjects, FToolMenuSection& Section);
+	static void AddMenuEntry();
+	static void RemoveMenuEntry();
 
-	void OnCreateProxy(TArray<FWeakObjectPtr> Objects) const;
+private:
 
-	TSharedRef<IAssetTypeActions> OriginalActions;
+	static void OnConstructMenu(FToolMenuSection& MenuSection);
+	static void OnExecuteAction(const FToolMenuContext& MenuContext);
 };
 
 #endif
diff --git a/Engine/Plugins/Enterprise/GLTFExporter/Source/GLTFExporter/Private/GLTFExporterModule.cpp b/Engine/Plugins/Enterprise/GLTFExporter/Source/GLTFExporter/Private/GLTFExporterModule.cpp
--- a/Engine/Plugins/Enterprise/GLTFExporter/Source/GLTFExporter/Private/GLTFExporterModule.cpp
+++ b/Engine/Plugins/Enterprise/GLTFExporter/Source/GLTFExporter/Private/GLTFExporterModule.cpp
@@ -30,63 +30,17 @@ public:
 		AddShaderSourceDirectoryMapping(TEXT("/Plugin/GLTFExporter"), PluginShaderDir);
 
 #if WITH_EDITOR
-		// TODO: UI and editor-only functions should be moved (as much as possible) into its own module
-		FCoreDelegates::OnPostEngineInit.AddRaw(this, &FGLTFExporterModule::PostEngineInit);
+		// TODO: UI and editor-only functions should be moved (as much as possible) into their own module
+		FCoreDelegates::OnPostEngineInit.AddStatic(&FGLTFProxyAssetActions::AddMenuEntry);
 #endif
 	}
 
-#if WITH_EDITOR
-	void PostEngineInit()
-	{
-		AddProxyAssetActions();
-	}
-
-	void AddProxyAssetActions()
-	{
-		if (FAssetToolsModule::IsModuleLoaded())
-		{
-			IAssetTools& AssetTools = FAssetToolsModule::GetModule().Get();
-
-			// TODO: add support for when GetAssetTypeActionsForClass returns null pointer
-			AssetTypeActionsArray.Add(MakeShared<FGLTFProxyAssetActions>(AssetTools.GetAssetTypeActionsForClass(UMaterial::StaticClass()).Pin().ToSharedRef()));
-			AssetTypeActionsArray.Add(MakeShared<FGLTFProxyAssetActions>(AssetTools.GetAssetTypeActionsForClass(UMaterialInstanceConstant::StaticClass()).Pin().ToSharedRef()));
-			AssetTypeActionsArray.Add(MakeShared<FGLTFProxyAssetActions>(AssetTools.GetAssetTypeActionsForClass(UMaterialInstanceDynamic::StaticClass()).Pin().ToSharedRef()));
-			AssetTypeActionsArray.Add(MakeShared<FGLTFProxyAssetActions>(AssetTools.GetAssetTypeActionsForClass(UMaterialInterface::StaticClass()).Pin().ToSharedRef()));
-
-			for (TSharedRef<IAssetTypeActions>& AssetTypeActions : AssetTypeActionsArray)
-			{
-				AssetTools.RegisterAssetTypeActions(AssetTypeActions);
-			}
-		}
-	}
-
-	void RemoveProxyAssetActions()
-	{
-		if (FAssetToolsModule::IsModuleLoaded())
-		{
-			IAssetTools& AssetTools = FAssetToolsModule::GetModule().Get();
-			for (TSharedRef<IAssetTypeActions>& AssetTypeActions : AssetTypeActionsArray)
-			{
-				AssetTools.UnregisterAssetTypeActions(AssetTypeActions);
-			}
-		}
-
-		AssetTypeActionsArray.Empty();
-	}
-#endif
-
 	virtual void ShutdownModule() override
 	{
 #if WITH_EDITOR
-		RemoveProxyAssetActions();
+		FGLTFProxyAssetActions::RemoveMenuEntry();
 #endif
 	}
-
-private:
-
-#if WITH_EDITOR
-	TArray<TSharedRef<IAssetTypeActions>> AssetTypeActionsArray;
-#endif
 };
 
 IMPLEMENT_MODULE(FGLTFExporterModule, GLTFExporter);
