Add missing call to open log after reimplementation of create proxy material asset action

diff --git a/Engine/Plugins/Enterprise/GLTFExporter/Source/GLTFExporter/Private/Actions/GLTFProxyAssetActions.cpp b/Engine/Plugins/Enterprise/GLTFExporter/Source/GLTFExporter/Private/Actions/GLTFProxyAssetActions.cpp
--- a/Engine/Plugins/Enterprise/GLTFExporter/Source/GLTFExporter/Private/Actions/GLTFProxyAssetActions.cpp
+++ b/Engine/Plugins/Enterprise/GLTFExporter/Source/GLTFExporter/Private/Actions/GLTFProxyAssetActions.cpp
@@ -70,6 +70,8 @@ void FGLTFProxyAssetActions::OnExecuteAction(const FToolMenuContext& MenuContext
 		ProxyFactory.RootPath = FPaths::GetPath(Material->GetPathName()) / TEXT("GLTF");
 		ProxyFactory.Create(Material);
 	}
+
+	ProxyFactory.OpenLog();
 }
 
 #undef LOCTEXT_NAMESPACE
