Remove old (no longer needed) hacks for accessing base properties when exporting proxy material to gltf

diff --git a/Engine/Plugins/Enterprise/GLTFExporter/Source/GLTFExporter/Private/Tasks/GLTFDelayedMaterialTasks.cpp b/Engine/Plugins/Enterprise/GLTFExporter/Source/GLTFExporter/Private/Tasks/GLTFDelayedMaterialTasks.cpp
--- a/Engine/Plugins/Enterprise/GLTFExporter/Source/GLTFExporter/Private/Tasks/GLTFDelayedMaterialTasks.cpp
+++ b/Engine/Plugins/Enterprise/GLTFExporter/Source/GLTFExporter/Private/Tasks/GLTFDelayedMaterialTasks.cpp
@@ -59,13 +59,6 @@ void FGLTFDelayedMaterialTask::Process()
 	JsonMaterial->AlphaCutoff = Material->GetOpacityMaskClipValue();
 	JsonMaterial->DoubleSided = Material->IsTwoSided();
 
-	if (const UMaterialInstance* MaterialInstance = Cast<UMaterialInstance>(Material))
-	{
-		// TODO: remove ugly hack for dynamic instances
-		FGLTFProxyMaterialUtilities::GetOpacityMaskClipValue(MaterialInstance, JsonMaterial->AlphaCutoff);
-		FGLTFProxyMaterialUtilities::GetTwoSided(MaterialInstance, JsonMaterial->DoubleSided);
-	}
-
 	ConvertShadingModel(JsonMaterial->ShadingModel);
 	ConvertAlphaMode(JsonMaterial->AlphaMode, JsonMaterial->BlendMode);
 
@@ -444,11 +437,7 @@ void FGLTFDelayedMaterialTask::ConvertShadingModel(EGLTFJsonShadingModel& OutSha
 
 void FGLTFDelayedMaterialTask::ConvertAlphaMode(EGLTFJsonAlphaMode& OutAlphaMode, EGLTFJsonBlendMode& OutBlendMode) const
 {
-	EBlendMode BlendMode = Material->GetBlendMode();
-	if (const UMaterialInstance* MaterialInstance = Cast<UMaterialInstance>(Material))
-	{
-		FGLTFProxyMaterialUtilities::GetBlendMode(MaterialInstance, BlendMode); // TODO: remove ugly hack for dynamic instances
-	}
+	const EBlendMode BlendMode = Material->GetBlendMode();
 
 	OutAlphaMode = FGLTFCoreUtilities::ConvertAlphaMode(BlendMode);
 	if (OutAlphaMode == EGLTFJsonAlphaMode::None)
