Remove old (no longer needed) hacks for accessing raw vertex data in gltf exporter

diff --git a/Engine/Plugins/Enterprise/GLTFExporter/Source/GLTFExporter/Private/Converters/GLTFBufferAdapter.cpp b/Engine/Plugins/Enterprise/GLTFExporter/Source/GLTFExporter/Private/Converters/GLTFBufferAdapter.cpp
--- a/Engine/Plugins/Enterprise/GLTFExporter/Source/GLTFExporter/Private/Converters/GLTFBufferAdapter.cpp
+++ b/Engine/Plugins/Enterprise/GLTFExporter/Source/GLTFExporter/Private/Converters/GLTFBufferAdapter.cpp
@@ -2,14 +2,11 @@
 
 #include "Converters/GLTFBufferAdapter.h"
 #include "Converters/GLTFBufferUtility.h"
-
+#include "RawIndexBuffer.h"
 #include "Rendering/ColorVertexBuffer.h"
 #include "Rendering/PositionVertexBuffer.h"
-#include "RawIndexBuffer.h"
-#include "RHI.h"
-#include "RHIResources.h"
-#include "Rendering/SkinWeightVertexBuffer.h"
 #include "Rendering/StaticMeshVertexBuffer.h"
+#include "Rendering/SkinWeightVertexBuffer.h"
 
 class FGLTFBufferAdapterCPU final : public IGLTFBufferAdapter
 {
@@ -61,28 +58,28 @@ TUniquePtr<IGLTFBufferAdapter> IGLTFBufferAdapter::GetIndices(const FRawStaticIn
 
 TUniquePtr<IGLTFBufferAdapter> IGLTFBufferAdapter::GetPositions(const FPositionVertexBuffer* VertexBuffer)
 {
-	const void* PositionData = const_cast<FPositionVertexBuffer*>(VertexBuffer)->GetVertexData();
+	const void* PositionData = VertexBuffer->GetVertexData();
 	if (PositionData != nullptr && FGLTFBufferUtility::HasCPUAccess(VertexBuffer)) return MakeUnique<FGLTFBufferAdapterCPU>(PositionData);
 	return MakeUnique<FGLTFBufferAdapterGPU>(VertexBuffer->VertexBufferRHI);
 }
 
 TUniquePtr<IGLTFBufferAdapter> IGLTFBufferAdapter::GetColors(const FColorVertexBuffer* VertexBuffer)
 {
-	const void* ColorData = const_cast<FColorVertexBuffer*>(VertexBuffer)->GetVertexData();
+	const void* ColorData = VertexBuffer->GetVertexData();
 	if (ColorData != nullptr && FGLTFBufferUtility::HasCPUAccess(VertexBuffer)) return MakeUnique<FGLTFBufferAdapterCPU>(ColorData);
 	return MakeUnique<FGLTFBufferAdapterGPU>(VertexBuffer->VertexBufferRHI);
 }
 
 TUniquePtr<IGLTFBufferAdapter> IGLTFBufferAdapter::GetTangents(const FStaticMeshVertexBuffer* VertexBuffer)
 {
-	const void* TangentData = const_cast<FStaticMeshVertexBuffer*>(VertexBuffer)->GetTangentData();
+	const void* TangentData = VertexBuffer->GetTangentData();
 	if (TangentData != nullptr && FGLTFBufferUtility::HasCPUAccess(VertexBuffer)) return MakeUnique<FGLTFBufferAdapterCPU>(TangentData);
 	return MakeUnique<FGLTFBufferAdapterGPU>(VertexBuffer->TangentsVertexBuffer.VertexBufferRHI);
 }
 
 TUniquePtr<IGLTFBufferAdapter> IGLTFBufferAdapter::GetUVs(const FStaticMeshVertexBuffer* VertexBuffer)
 {
-	const void* UVData = const_cast<FStaticMeshVertexBuffer*>(VertexBuffer)->GetTexCoordData();
+	const void* UVData = VertexBuffer->GetTexCoordData();
 	if (UVData != nullptr && FGLTFBufferUtility::HasCPUAccess(VertexBuffer)) return MakeUnique<FGLTFBufferAdapterCPU>(UVData);
 	return MakeUnique<FGLTFBufferAdapterGPU>(VertexBuffer->TexCoordVertexBuffer.VertexBufferRHI);
 }
@@ -90,7 +87,7 @@ TUniquePtr<IGLTFBufferAdapter> IGLTFBufferAdapter::GetUVs(const FStaticMeshVerte
 TUniquePtr<IGLTFBufferAdapter> IGLTFBufferAdapter::GetInfluences(const FSkinWeightVertexBuffer* VertexBuffer)
 {
 	const FSkinWeightDataVertexBuffer* InfluenceBuffer = VertexBuffer->GetDataVertexBuffer();
-	const void* InfluenceData = FGLTFBufferUtility::GetCPUBuffer(InfluenceBuffer);
+	const void* InfluenceData = InfluenceBuffer->GetWeightData();
 	if (InfluenceData != nullptr && FGLTFBufferUtility::HasCPUAccess(VertexBuffer)) return MakeUnique<FGLTFBufferAdapterCPU>(InfluenceData);
 	return MakeUnique<FGLTFBufferAdapterGPU>(InfluenceBuffer->VertexBufferRHI);
 }
diff --git a/Engine/Plugins/Enterprise/GLTFExporter/Source/GLTFExporter/Private/Converters/GLTFBufferUtility.cpp b/Engine/Plugins/Enterprise/GLTFExporter/Source/GLTFExporter/Private/Converters/GLTFBufferUtility.cpp
--- a/Engine/Plugins/Enterprise/GLTFExporter/Source/GLTFExporter/Private/Converters/GLTFBufferUtility.cpp
+++ b/Engine/Plugins/Enterprise/GLTFExporter/Source/GLTFExporter/Private/Converters/GLTFBufferUtility.cpp
@@ -8,16 +8,7 @@ bool FGLTFBufferUtility::HasCPUAccess(const FRawStaticIndexBuffer* IndexBuffer)
 #if WITH_EDITOR
 	return true;
 #else
-	struct FRawStaticIndexBufferHack : FIndexBuffer
-	{
-		TResourceArray<uint8, INDEXBUFFER_ALIGNMENT> IndexStorage;
-		int32 CachedNumIndices;
-		bool b32Bit;
-		bool bShouldExpandTo32Bit;
-	};
-
-	static_assert(sizeof(FRawStaticIndexBufferHack) == sizeof(FRawStaticIndexBuffer), "FRawStaticIndexBufferHack memory layout doesn't match FRawStaticIndexBuffer");
-	return reinterpret_cast<const FRawStaticIndexBufferHack*>(IndexBuffer)->IndexStorage.GetAllowCPUAccess();
+	return IndexBuffer->GetAllowCPUAccess();
 #endif
 }
 
@@ -35,19 +26,7 @@ bool FGLTFBufferUtility::HasCPUAccess(const FPositionVertexBuffer* VertexBuffer)
 #if WITH_EDITOR
 	return true;
 #else
-	struct FPositionVertexBufferHack : FVertexBuffer
-	{
-		FShaderResourceViewRHIRef PositionComponentSRV;
-		TMemoryImagePtr<FStaticMeshVertexDataInterface> VertexData;
-		uint8* Data;
-		uint32 Stride;
-		uint32 NumVertices;
-		bool bNeedsCPUAccess;
-	};
-
-	static_assert(sizeof(FPositionVertexBufferHack) == sizeof(FPositionVertexBuffer), "FPositionVertexBufferHack memory layout doesn't match FPositionVertexBuffer");
-	const FStaticMeshVertexDataInterface* VertexData = reinterpret_cast<const FPositionVertexBufferHack*>(VertexBuffer)->VertexData;
-	return VertexData != nullptr && VertexData->GetAllowCPUAccess();
+	return VertexBuffer->GetAllowCPUAccess();
 #endif
 }
 
@@ -56,19 +35,7 @@ bool FGLTFBufferUtility::HasCPUAccess(const FColorVertexBuffer* VertexBuffer)
 #if WITH_EDITOR
 	return true;
 #else
-	struct FColorVertexBufferHack : FVertexBuffer
-	{
-		FStaticMeshVertexDataInterface* VertexData;
-		FShaderResourceViewRHIRef ColorComponentsSRV;
-		uint8* Data;
-		uint32 Stride;
-		uint32 NumVertices;
-		bool NeedsCPUAccess;
-	};
-
-	static_assert(sizeof(FColorVertexBufferHack) == sizeof(FColorVertexBuffer), "FColorVertexBufferHack memory layout doesn't match FColorVertexBuffer");
-	const FStaticMeshVertexDataInterface* VertexData = reinterpret_cast<const FColorVertexBufferHack*>(VertexBuffer)->VertexData;
-	return VertexData != nullptr && VertexData->GetAllowCPUAccess();
+	return VertexBuffer->GetAllowCPUAccess();
 #endif
 }
 
@@ -100,11 +67,6 @@ const void* FGLTFBufferUtility::GetCPUBuffer(const FRawStaticIndexBuffer16or32In
 	return IndexBuffer->GetResourceDataSize() > 0 ? const_cast<FRawStaticIndexBuffer16or32Interface*>(IndexBuffer)->GetPointerTo(0) : nullptr;
 }
 
-const void* FGLTFBufferUtility::GetCPUBuffer(const FSkinWeightDataVertexBuffer* VertexBuffer)
-{
-	return VertexBuffer->GetWeightData();
-}
-
 const void* FGLTFBufferUtility::GetCPUBuffer(const FSkinWeightLookupVertexBuffer* VertexBuffer)
 {
 	struct FSkinWeightLookupVertexBufferHack : FVertexBuffer
diff --git a/Engine/Plugins/Enterprise/GLTFExporter/Source/GLTFExporter/Private/Converters/GLTFBufferUtility.h b/Engine/Plugins/Enterprise/GLTFExporter/Source/GLTFExporter/Private/Converters/GLTFBufferUtility.h
--- a/Engine/Plugins/Enterprise/GLTFExporter/Source/GLTFExporter/Private/Converters/GLTFBufferUtility.h
+++ b/Engine/Plugins/Enterprise/GLTFExporter/Source/GLTFExporter/Private/Converters/GLTFBufferUtility.h
@@ -25,7 +25,6 @@ struct FGLTFBufferUtility
 
 	static const void* GetCPUBuffer(const FRawStaticIndexBuffer* IndexBuffer);
 	static const void* GetCPUBuffer(const FRawStaticIndexBuffer16or32Interface* IndexBuffer);
-	static const void* GetCPUBuffer(const FSkinWeightDataVertexBuffer* VertexBuffer);
 	static const void* GetCPUBuffer(const FSkinWeightLookupVertexBuffer* VertexBuffer);
 
 	static void ReadRHIBuffer(FRHIBuffer* SourceBuffer, TArray<uint8>& OutData);
