Remove old (no longer needed) hacks for accessing transform channels when exporting level sequences to gltf

diff --git a/Engine/Plugins/Enterprise/GLTFExporter/Source/GLTFExporter/Private/Converters/GLTFBoneUtility.cpp b/Engine/Plugins/Enterprise/GLTFExporter/Source/GLTFExporter/Private/Converters/GLTFBoneUtility.cpp
--- a/Engine/Plugins/Enterprise/GLTFExporter/Source/GLTFExporter/Private/Converters/GLTFBoneUtility.cpp
+++ b/Engine/Plugins/Enterprise/GLTFExporter/Source/GLTFExporter/Private/Converters/GLTFBoneUtility.cpp
@@ -89,24 +89,3 @@ void FGLTFBoneUtility::GetBoneTransformsByFrame(const UAnimSequence* AnimSequenc
 	Curve.Empty();
 	Attributes.Empty();
 }
-
-const FMovieSceneDoubleChannel* FGLTFBoneUtility::GetTranslationChannels(const UMovieScene3DTransformSection* TransformSection)
-{
-	// TODO: fix ugly workaround hack in engine api to access channels without GetChannelProxy (which doesn't work properly in UE5)
-	static const FProperty* Property = UMovieScene3DTransformSection::StaticClass()->FindPropertyByName(TEXT("Translation"));
-	return Property->ContainerPtrToValuePtr<FMovieSceneDoubleChannel>(TransformSection);
-}
-
-const FMovieSceneDoubleChannel* FGLTFBoneUtility::GetRotationChannels(const UMovieScene3DTransformSection* TransformSection)
-{
-	// TODO: fix ugly workaround hack in engine api to access channels without GetChannelProxy (which doesn't work properly in UE5)
-	static const FProperty* Property = UMovieScene3DTransformSection::StaticClass()->FindPropertyByName(TEXT("Rotation"));
-	return Property->ContainerPtrToValuePtr<FMovieSceneDoubleChannel>(TransformSection);
-}
-
-const FMovieSceneDoubleChannel* FGLTFBoneUtility::GetScaleChannels(const UMovieScene3DTransformSection* TransformSection)
-{
-	// TODO: fix ugly workaround hack in engine api to access channels without GetChannelProxy (which doesn't work properly in UE5)
-	static const FProperty* Property = UMovieScene3DTransformSection::StaticClass()->FindPropertyByName(TEXT("Scale"));
-	return Property->ContainerPtrToValuePtr<FMovieSceneDoubleChannel>(TransformSection);
-}
diff --git a/Engine/Plugins/Enterprise/GLTFExporter/Source/GLTFExporter/Private/Converters/GLTFBoneUtility.h b/Engine/Plugins/Enterprise/GLTFExporter/Source/GLTFExporter/Private/Converters/GLTFBoneUtility.h
--- a/Engine/Plugins/Enterprise/GLTFExporter/Source/GLTFExporter/Private/Converters/GLTFBoneUtility.h
+++ b/Engine/Plugins/Enterprise/GLTFExporter/Source/GLTFExporter/Private/Converters/GLTFBoneUtility.h
@@ -20,8 +20,4 @@ struct FGLTFBoneUtility
 	static void GetBoneIndices(const USkeleton* Skeleton, TArray<FBoneIndexType>& OutBoneIndices);
 
 	static void GetBoneTransformsByFrame(const UAnimSequence* AnimSequence, const TArray<float>& FrameTimestamps, const TArray<FBoneIndexType>& BoneIndices, TArray<TArray<FTransform>>& OutBoneTransformsByFrame);
-
-	static const FMovieSceneDoubleChannel* GetTranslationChannels(const UMovieScene3DTransformSection* TransformSection);
-	static const FMovieSceneDoubleChannel* GetRotationChannels(const UMovieScene3DTransformSection* TransformSection);
-	static const FMovieSceneDoubleChannel* GetScaleChannels(const UMovieScene3DTransformSection* TransformSection);
 };
diff --git a/Engine/Plugins/Enterprise/GLTFExporter/Source/GLTFExporter/Private/Tasks/GLTFDelayedAnimationTasks.cpp b/Engine/Plugins/Enterprise/GLTFExporter/Source/GLTFExporter/Private/Tasks/GLTFDelayedAnimationTasks.cpp
--- a/Engine/Plugins/Enterprise/GLTFExporter/Source/GLTFExporter/Private/Tasks/GLTFDelayedAnimationTasks.cpp
+++ b/Engine/Plugins/Enterprise/GLTFExporter/Source/GLTFExporter/Private/Tasks/GLTFDelayedAnimationTasks.cpp
@@ -243,14 +243,18 @@ void FGLTFDelayedLevelSequenceTask::Process()
 
 					// TODO: do we need to account for TransformSection->GetOffsetTime() ?
 
+					TArrayView<FMovieSceneDoubleChannel*> Channels = TransformSection->GetChannelProxy().GetChannels<FMovieSceneDoubleChannel>();
 					const EMovieSceneTransformChannel ChannelMask = TransformSection->GetMask().GetChannels();
 
 					if (EnumHasAnyFlags(ChannelMask, EMovieSceneTransformChannel::Translation))
 					{
-						const FMovieSceneDoubleChannel* Channels = FGLTFBoneUtility::GetTranslationChannels(TransformSection);
-						const bool MaskX = EnumHasAllFlags(ChannelMask, EMovieSceneTransformChannel::TranslationX);
-						const bool MaskY = EnumHasAllFlags(ChannelMask, EMovieSceneTransformChannel::TranslationY);
-						const bool MaskZ = EnumHasAllFlags(ChannelMask, EMovieSceneTransformChannel::TranslationZ);
+						constexpr int32 IndexX = 0;
+						constexpr int32 IndexY = 1;
+						constexpr int32 IndexZ = 2;
+
+						const bool MaskX = EnumHasAllFlags(ChannelMask, EMovieSceneTransformChannel::TranslationX) && Channels.IsValidIndex(IndexX);
+						const bool MaskY = EnumHasAllFlags(ChannelMask, EMovieSceneTransformChannel::TranslationY) && Channels.IsValidIndex(IndexY);
+						const bool MaskZ = EnumHasAllFlags(ChannelMask, EMovieSceneTransformChannel::TranslationZ) && Channels.IsValidIndex(IndexZ);
 
 						TArray<FGLTFVector3> Translations;
 						Translations.AddUninitialized(FrameCount);
@@ -260,9 +264,9 @@ void FGLTFDelayedLevelSequenceTask::Process()
 							const FFrameTime& FrameTime = FrameTimes[Frame];
 
 							FVector3f Translation = DefaultTranslation;
-							if (MaskX) Channels[0].Evaluate(FrameTime, Translation.X);
-							if (MaskY) Channels[1].Evaluate(FrameTime, Translation.Y);
-							if (MaskZ) Channels[2].Evaluate(FrameTime, Translation.Z);
+							if (MaskX) Channels[IndexX]->Evaluate(FrameTime, Translation.X);
+							if (MaskY) Channels[IndexY]->Evaluate(FrameTime, Translation.Y);
+							if (MaskZ) Channels[IndexZ]->Evaluate(FrameTime, Translation.Z);
 
 							Translations[Frame] = FGLTFCoreUtilities::ConvertPosition(Translation, Builder.ExportOptions->ExportUniformScale);
 						}
@@ -287,10 +291,13 @@ void FGLTFDelayedLevelSequenceTask::Process()
 
 					if (EnumHasAnyFlags(ChannelMask, EMovieSceneTransformChannel::Rotation))
 					{
-						const FMovieSceneDoubleChannel* Channels = FGLTFBoneUtility::GetRotationChannels(TransformSection);
-						const bool MaskX = EnumHasAllFlags(ChannelMask, EMovieSceneTransformChannel::RotationX);
-						const bool MaskY = EnumHasAllFlags(ChannelMask, EMovieSceneTransformChannel::RotationY);
-						const bool MaskZ = EnumHasAllFlags(ChannelMask, EMovieSceneTransformChannel::RotationZ);
+						constexpr int32 IndexX = 3;
+						constexpr int32 IndexY = 4;
+						constexpr int32 IndexZ = 5;
+
+						const bool MaskX = EnumHasAllFlags(ChannelMask, EMovieSceneTransformChannel::RotationX) && Channels.IsValidIndex(IndexX);
+						const bool MaskY = EnumHasAllFlags(ChannelMask, EMovieSceneTransformChannel::RotationY) && Channels.IsValidIndex(IndexY);
+						const bool MaskZ = EnumHasAllFlags(ChannelMask, EMovieSceneTransformChannel::RotationZ) && Channels.IsValidIndex(IndexZ);
 
 						TArray<FGLTFQuaternion> Rotations;
 						Rotations.AddUninitialized(FrameCount);
@@ -300,9 +307,9 @@ void FGLTFDelayedLevelSequenceTask::Process()
 							const FFrameTime& FrameTime = FrameTimes[Frame];
 
 							FRotator3f Rotator = DefaultRotator;
-							if (MaskX) Channels[0].Evaluate(FrameTime, Rotator.Roll);
-							if (MaskY) Channels[1].Evaluate(FrameTime, Rotator.Pitch);
-							if (MaskZ) Channels[2].Evaluate(FrameTime, Rotator.Yaw);
+							if (MaskX) Channels[IndexX]->Evaluate(FrameTime, Rotator.Roll);
+							if (MaskY) Channels[IndexY]->Evaluate(FrameTime, Rotator.Pitch);
+							if (MaskZ) Channels[IndexZ]->Evaluate(FrameTime, Rotator.Yaw);
 
 							Rotations[Frame] = FGLTFCoreUtilities::ConvertRotation(Rotator.Quaternion());
 						}
@@ -327,10 +334,13 @@ void FGLTFDelayedLevelSequenceTask::Process()
 
 					if (EnumHasAnyFlags(ChannelMask, EMovieSceneTransformChannel::Scale))
 					{
-						const FMovieSceneDoubleChannel* Channels = FGLTFBoneUtility::GetScaleChannels(TransformSection);
-						const bool MaskX = EnumHasAllFlags(ChannelMask, EMovieSceneTransformChannel::ScaleX);
-						const bool MaskY = EnumHasAllFlags(ChannelMask, EMovieSceneTransformChannel::ScaleY);
-						const bool MaskZ = EnumHasAllFlags(ChannelMask, EMovieSceneTransformChannel::ScaleZ);
+						constexpr int32 IndexX = 6;
+						constexpr int32 IndexY = 7;
+						constexpr int32 IndexZ = 8;
+
+						const bool MaskX = EnumHasAllFlags(ChannelMask, EMovieSceneTransformChannel::ScaleX) && Channels.IsValidIndex(IndexX);
+						const bool MaskY = EnumHasAllFlags(ChannelMask, EMovieSceneTransformChannel::ScaleY) && Channels.IsValidIndex(IndexY);
+						const bool MaskZ = EnumHasAllFlags(ChannelMask, EMovieSceneTransformChannel::ScaleZ) && Channels.IsValidIndex(IndexZ);
 
 						TArray<FGLTFVector3> Scales;
 						Scales.AddUninitialized(FrameCount);
@@ -340,9 +350,9 @@ void FGLTFDelayedLevelSequenceTask::Process()
 							const FFrameTime& FrameTime = FrameTimes[Frame];
 
 							FVector3f Scale = DefaultScale;
-							if (MaskX) Channels[0].Evaluate(FrameTime, Scale.X);
-							if (MaskY) Channels[1].Evaluate(FrameTime, Scale.Y);
-							if (MaskZ) Channels[2].Evaluate(FrameTime, Scale.Z);
+							if (MaskX) Channels[IndexX]->Evaluate(FrameTime, Scale.X);
+							if (MaskY) Channels[IndexY]->Evaluate(FrameTime, Scale.Y);
+							if (MaskZ) Channels[IndexZ]->Evaluate(FrameTime, Scale.Z);
 
 							Scales[Frame] = FGLTFCoreUtilities::ConvertScale(Scale);
 						}
