Remove all custom Epic extensions in gltf exporter, to reduce code complexity and simplify testing before adding any more features

diff --git a/Engine/Plugins/Enterprise/GLTFExporter/Shaders/Private/CubemapFacePS.usf b/Engine/Plugins/Enterprise/GLTFExporter/Shaders/Private/CubemapFacePS.usf
deleted file mode 100644
diff --git a/Engine/Plugins/Enterprise/GLTFExporter/Source/GLTFExporter/Private/Builders/GLTFBuilder.cpp b/Engine/Plugins/Enterprise/GLTFExporter/Source/GLTFExporter/Private/Builders/GLTFBuilder.cpp
--- a/Engine/Plugins/Enterprise/GLTFExporter/Source/GLTFExporter/Private/Builders/GLTFBuilder.cpp
+++ b/Engine/Plugins/Enterprise/GLTFExporter/Source/GLTFExporter/Private/Builders/GLTFBuilder.cpp
@@ -76,19 +76,6 @@ TextureAddress FGLTFBuilder::GetBakeTilingForMaterialProperty(const UMaterialInt
 	return UGLTFMaterialExportOptions::GetBakeTilingForPropertyGroup(Material, PropertyGroup, DefaultValue);
 }
 
-EGLTFJsonHDREncoding FGLTFBuilder::GetTextureHDREncoding() const
-{
-	switch (ExportOptions->TextureHDREncoding)
-	{
-		case EGLTFTextureHDREncoding::None: return EGLTFJsonHDREncoding::None;
-		case EGLTFTextureHDREncoding::RGBM: return EGLTFJsonHDREncoding::RGBM;
-		// TODO: add more encodings (like RGBE) when viewer supports them
-		default:
-			checkNoEntry();
-			return EGLTFJsonHDREncoding::None;
-	}
-}
-
 bool FGLTFBuilder::ShouldExportLight(EComponentMobility::Type LightMobility) const
 {
 	const EGLTFSceneMobility AllowedMobility = static_cast<EGLTFSceneMobility>(ExportOptions->ExportLights);
diff --git a/Engine/Plugins/Enterprise/GLTFExporter/Source/GLTFExporter/Private/Builders/GLTFContainerBuilder.cpp b/Engine/Plugins/Enterprise/GLTFExporter/Source/GLTFExporter/Private/Builders/GLTFContainerBuilder.cpp
--- a/Engine/Plugins/Enterprise/GLTFExporter/Source/GLTFExporter/Private/Builders/GLTFContainerBuilder.cpp
+++ b/Engine/Plugins/Enterprise/GLTFExporter/Source/GLTFExporter/Private/Builders/GLTFContainerBuilder.cpp
@@ -21,18 +21,6 @@ void FGLTFContainerBuilder::WriteInternalArchive(FArchive& Archive)
 	{
 		WriteJsonArchive(Archive);
 	}
-
-	const TSet<EGLTFJsonExtension> CustomExtensions = GetCustomExtensionsUsed();
-	if (CustomExtensions.Num() > 0)
-	{
-		const FString ExtensionsString = FString::JoinBy(CustomExtensions, TEXT(", "),
-			[](EGLTFJsonExtension Extension)
-		{
-			return FGLTFJsonUtilities::GetValue(Extension);
-		});
-
-		LogWarning(FString::Printf(TEXT("Export uses some extensions that may only be supported in Unreal's glTF viewer: %s"), *ExtensionsString));
-	}
 }
 
 bool FGLTFContainerBuilder::WriteAllFiles(const FString& DirPath, uint32 WriteFlags)
diff --git a/Engine/Plugins/Enterprise/GLTFExporter/Source/GLTFExporter/Private/Builders/GLTFConvertBuilder.cpp b/Engine/Plugins/Enterprise/GLTFExporter/Source/GLTFExporter/Private/Builders/GLTFConvertBuilder.cpp
--- a/Engine/Plugins/Enterprise/GLTFExporter/Source/GLTFExporter/Private/Builders/GLTFConvertBuilder.cpp
+++ b/Engine/Plugins/Enterprise/GLTFExporter/Source/GLTFExporter/Private/Builders/GLTFConvertBuilder.cpp
@@ -5,9 +5,7 @@
 #include "Components/SkeletalMeshComponent.h"
 #include "Components/StaticMeshComponent.h"
 #include "Engine/Texture2D.h"
-#include "Engine/TextureCube.h"
 #include "Engine/TextureRenderTarget2D.h"
-#include "Engine/TextureRenderTargetCube.h"
 
 FGLTFConvertBuilder::FGLTFConvertBuilder(const FString& FileName, const UGLTFExportOptions* ExportOptions, const TSet<AActor*>& SelectedActors)
 	: FGLTFBufferBuilder(FileName, ExportOptions)
@@ -262,31 +260,11 @@ FGLTFJsonTexture* FGLTFConvertBuilder::AddUniqueTexture(const UTexture2D* Textur
 	return AddUniqueTexture(Texture, Texture->SRGB);
 }
 
-FGLTFJsonTexture* FGLTFConvertBuilder::AddUniqueTexture(const UTextureCube* Texture, ECubeFace CubeFace)
-{
-	return AddUniqueTexture(Texture, CubeFace, Texture->SRGB);
-}
-
 FGLTFJsonTexture* FGLTFConvertBuilder::AddUniqueTexture(const UTextureRenderTarget2D* Texture)
 {
 	return AddUniqueTexture(Texture, Texture->SRGB);
 }
 
-FGLTFJsonTexture* FGLTFConvertBuilder::AddUniqueTexture(const UTextureRenderTargetCube* Texture, ECubeFace CubeFace)
-{
-	return AddUniqueTexture(Texture, CubeFace, Texture->SRGB);
-}
-
-FGLTFJsonTexture* FGLTFConvertBuilder::AddUniqueTexture(const ULightMapTexture2D* Texture)
-{
-	if (Texture == nullptr)
-	{
-		return nullptr;
-	}
-
-	return TextureLightMapConverter->GetOrAdd(Texture);
-}
-
 FGLTFJsonTexture* FGLTFConvertBuilder::AddUniqueTexture(const UTexture* Texture, bool bToSRGB)
 {
 	if (const UTexture2D* Texture2D = Cast<UTexture2D>(Texture))
@@ -312,16 +290,6 @@ FGLTFJsonTexture* FGLTFConvertBuilder::AddUniqueTexture(const UTexture2D* Textur
 	return Texture2DConverter->GetOrAdd(Texture, bToSRGB);
 }
 
-FGLTFJsonTexture* FGLTFConvertBuilder::AddUniqueTexture(const UTextureCube* Texture, ECubeFace CubeFace, bool bToSRGB)
-{
-	if (Texture == nullptr)
-	{
-		return nullptr;
-	}
-
-	return TextureCubeConverter->GetOrAdd(Texture, CubeFace, bToSRGB);
-}
-
 FGLTFJsonTexture* FGLTFConvertBuilder::AddUniqueTexture(const UTextureRenderTarget2D* Texture, bool bToSRGB)
 {
 	if (Texture == nullptr)
@@ -332,16 +300,6 @@ FGLTFJsonTexture* FGLTFConvertBuilder::AddUniqueTexture(const UTextureRenderTarg
 	return TextureRenderTarget2DConverter->GetOrAdd(Texture, bToSRGB);
 }
 
-FGLTFJsonTexture* FGLTFConvertBuilder::AddUniqueTexture(const UTextureRenderTargetCube* Texture, ECubeFace CubeFace, bool bToSRGB)
-{
-	if (Texture == nullptr)
-	{
-		return nullptr;
-	}
-
-	return TextureRenderTargetCubeConverter->GetOrAdd(Texture, CubeFace, bToSRGB);
-}
-
 FGLTFJsonImage* FGLTFConvertBuilder::AddUniqueImage(TGLTFSharedArray<FColor>& Pixels, FIntPoint Size, bool bIgnoreAlpha, EGLTFTextureType Type, const FString& Name)
 {
 	return ImageConverter->GetOrAdd(Name, Type, bIgnoreAlpha, Size, Pixels);
@@ -504,46 +462,6 @@ FGLTFJsonLight* FGLTFConvertBuilder::AddUniqueLight(const ULightComponent* Light
 	return LightConverter->GetOrAdd(LightComponent);
 }
 
-FGLTFJsonBackdrop* FGLTFConvertBuilder::AddUniqueBackdrop(const AActor* BackdropActor)
-{
-	if (BackdropActor == nullptr)
-	{
-		return nullptr;
-	}
-
-	return BackdropConverter->GetOrAdd(BackdropActor);
-}
-
-FGLTFJsonLightMap* FGLTFConvertBuilder::AddUniqueLightMap(const UStaticMeshComponent* StaticMeshComponent)
-{
-	if (StaticMeshComponent == nullptr)
-	{
-		return nullptr;
-	}
-
-	return LightMapConverter->GetOrAdd(StaticMeshComponent);
-}
-
-FGLTFJsonSkySphere* FGLTFConvertBuilder::AddUniqueSkySphere(const AActor* SkySphereActor)
-{
-	if (SkySphereActor == nullptr)
-	{
-		return nullptr;
-	}
-
-	return SkySphereConverter->GetOrAdd(SkySphereActor);
-}
-
-FGLTFJsonEpicLevelVariantSets* FGLTFConvertBuilder::AddUniqueEpicLevelVariantSets(const ULevelVariantSets* LevelVariantSets)
-{
-	if (LevelVariantSets == nullptr)
-	{
-		return nullptr;
-	}
-
-	return EpicLevelVariantSetsConverter->GetOrAdd(LevelVariantSets);
-}
-
 FGLTFJsonKhrMaterialVariant* FGLTFConvertBuilder::AddUniqueKhrMaterialVariant(const UVariant* Variant)
 {
 	if (Variant == nullptr)
diff --git a/Engine/Plugins/Enterprise/GLTFExporter/Source/GLTFExporter/Private/Builders/GLTFJsonBuilder.cpp b/Engine/Plugins/Enterprise/GLTFExporter/Source/GLTFExporter/Private/Builders/GLTFJsonBuilder.cpp
--- a/Engine/Plugins/Enterprise/GLTFExporter/Source/GLTFExporter/Private/Builders/GLTFJsonBuilder.cpp
+++ b/Engine/Plugins/Enterprise/GLTFExporter/Source/GLTFExporter/Private/Builders/GLTFJsonBuilder.cpp
@@ -17,21 +17,6 @@ void FGLTFJsonBuilder::WriteJsonArchive(FArchive& Archive)
 	JsonRoot.WriteJson(Archive, !bIsGLB, ExportOptions->bSkipNearDefaultValues ? KINDA_SMALL_NUMBER : 0);
 }
 
-TSet<EGLTFJsonExtension> FGLTFJsonBuilder::GetCustomExtensionsUsed() const
-{
-	TSet<EGLTFJsonExtension> CustomExtensions;
-
-	for (EGLTFJsonExtension Extension : JsonRoot.Extensions.Used)
-	{
-		if (IsCustomExtension(Extension))
-		{
-			CustomExtensions.Add(Extension);
-		}
-	}
-
-	return CustomExtensions;
-}
-
 void FGLTFJsonBuilder::AddExtension(EGLTFJsonExtension Extension, bool bIsRequired)
 {
 	JsonRoot.Extensions.Used.Add(Extension);
@@ -106,31 +91,11 @@ FGLTFJsonTexture* FGLTFJsonBuilder::AddTexture()
 	return JsonRoot.Textures.Add();
 }
 
-FGLTFJsonBackdrop* FGLTFJsonBuilder::AddBackdrop()
-{
-	return JsonRoot.Backdrops.Add();
-}
-
 FGLTFJsonLight* FGLTFJsonBuilder::AddLight()
 {
 	return JsonRoot.Lights.Add();
 }
 
-FGLTFJsonLightMap* FGLTFJsonBuilder::AddLightMap()
-{
-	return JsonRoot.LightMaps.Add();
-}
-
-FGLTFJsonSkySphere* FGLTFJsonBuilder::AddSkySphere()
-{
-	return JsonRoot.SkySpheres.Add();
-}
-
-FGLTFJsonEpicLevelVariantSets* FGLTFJsonBuilder::AddEpicLevelVariantSets()
-{
-	return JsonRoot.EpicLevelVariantSets.Add();
-}
-
 FGLTFJsonKhrMaterialVariant* FGLTFJsonBuilder::AddKhrMaterialVariant()
 {
 	return JsonRoot.KhrMaterialVariants.Add();
@@ -150,11 +115,3 @@ FString FGLTFJsonBuilder::GetGeneratorString() const
 		? TEXT(EPIC_PRODUCT_NAME) TEXT(" ") ENGINE_VERSION_STRING TEXT(" ") + PluginDescriptor.FriendlyName + TEXT(" ") + PluginDescriptor.VersionName
 		: TEXT(EPIC_PRODUCT_NAME) TEXT(" ") + PluginDescriptor.FriendlyName;
 }
-
-bool FGLTFJsonBuilder::IsCustomExtension(EGLTFJsonExtension Extension)
-{
-	const TCHAR CustomPrefix[] = TEXT("EPIC_");
-
-	const TCHAR* ExtensionString = FGLTFJsonUtilities::GetValue(Extension);
-	return FCString::Strncmp(ExtensionString, CustomPrefix, GetNum(CustomPrefix) - 1) == 0;
-}
diff --git a/Engine/Plugins/Enterprise/GLTFExporter/Source/GLTFExporter/Private/Converters/GLTFAnimationConverters.cpp b/Engine/Plugins/Enterprise/GLTFExporter/Source/GLTFExporter/Private/Converters/GLTFAnimationConverters.cpp
--- a/Engine/Plugins/Enterprise/GLTFExporter/Source/GLTFExporter/Private/Converters/GLTFAnimationConverters.cpp
+++ b/Engine/Plugins/Enterprise/GLTFExporter/Source/GLTFExporter/Private/Converters/GLTFAnimationConverters.cpp
@@ -55,18 +55,7 @@ FGLTFJsonAnimation* FGLTFAnimationDataConverter::Convert(FGLTFJsonNode* RootNode
 		return nullptr;
 	}
 
-	FGLTFJsonAnimation* JsonAnimation = Builder.AddUniqueAnimation(RootNode, SkeletalMesh, AnimSequence);
-	if (JsonAnimation != nullptr && Builder.ExportOptions->bExportPlaybackSettings)
-	{
-		FGLTFJsonAnimationPlayback& JsonPlayback = JsonAnimation->Playback;
-
-		JsonPlayback.bLoop = SkeletalMeshComponent->AnimationData.bSavedLooping;
-		JsonPlayback.bAutoPlay = SkeletalMeshComponent->AnimationData.bSavedPlaying;
-		JsonPlayback.PlayRate = SkeletalMeshComponent->AnimationData.SavedPlayRate;
-		JsonPlayback.StartTime = SkeletalMeshComponent->AnimationData.SavedPosition;
-	}
-
-	return JsonAnimation;
+	return Builder.AddUniqueAnimation(RootNode, SkeletalMesh, AnimSequence);
 }
 
 FGLTFJsonAnimation* FGLTFLevelSequenceConverter::Convert(const ULevel* Level, const ULevelSequence* LevelSequence)
@@ -99,17 +88,5 @@ FGLTFJsonAnimation* FGLTFLevelSequenceDataConverter::Convert(const ALevelSequenc
 #endif
 	}
 
-	FGLTFJsonAnimation* JsonAnimation = Builder.AddUniqueAnimation(Level, LevelSequence);
-	if (JsonAnimation != nullptr && Builder.ExportOptions->bExportPlaybackSettings)
-	{
-		FGLTFJsonAnimationPlayback& JsonPlayback = JsonAnimation->Playback;
-
-		// TODO: report warning if loop count is not 0 or -1 (infinite)
-		JsonPlayback.bLoop = LevelSequenceActor->PlaybackSettings.LoopCount.Value != 0;
-		JsonPlayback.bAutoPlay = LevelSequenceActor->PlaybackSettings.bAutoPlay;
-		JsonPlayback.PlayRate = LevelSequenceActor->PlaybackSettings.PlayRate;
-		JsonPlayback.StartTime = LevelSequenceActor->PlaybackSettings.StartTime;
-	}
-
-	return JsonAnimation;
+	return Builder.AddUniqueAnimation(Level, LevelSequence);
 }
diff --git a/Engine/Plugins/Enterprise/GLTFExporter/Source/GLTFExporter/Private/Converters/GLTFBackdropConverters.cpp b/Engine/Plugins/Enterprise/GLTFExporter/Source/GLTFExporter/Private/Converters/GLTFBackdropConverters.cpp
deleted file mode 100644
diff --git a/Engine/Plugins/Enterprise/GLTFExporter/Source/GLTFExporter/Private/Converters/GLTFBlueprintUtility.cpp b/Engine/Plugins/Enterprise/GLTFExporter/Source/GLTFExporter/Private/Converters/GLTFBlueprintUtility.cpp
deleted file mode 100644
diff --git a/Engine/Plugins/Enterprise/GLTFExporter/Source/GLTFExporter/Private/Converters/GLTFBlueprintUtility.h b/Engine/Plugins/Enterprise/GLTFExporter/Source/GLTFExporter/Private/Converters/GLTFBlueprintUtility.h
deleted file mode 100644
diff --git a/Engine/Plugins/Enterprise/GLTFExporter/Source/GLTFExporter/Private/Converters/GLTFCubemapFacePreview.cpp b/Engine/Plugins/Enterprise/GLTFExporter/Source/GLTFExporter/Private/Converters/GLTFCubemapFacePreview.cpp
deleted file mode 100644
diff --git a/Engine/Plugins/Enterprise/GLTFExporter/Source/GLTFExporter/Private/Converters/GLTFCubemapFacePreview.h b/Engine/Plugins/Enterprise/GLTFExporter/Source/GLTFExporter/Private/Converters/GLTFCubemapFacePreview.h
deleted file mode 100644
diff --git a/Engine/Plugins/Enterprise/GLTFExporter/Source/GLTFExporter/Private/Converters/GLTFCurveUtility.cpp b/Engine/Plugins/Enterprise/GLTFExporter/Source/GLTFExporter/Private/Converters/GLTFCurveUtility.cpp
deleted file mode 100644
diff --git a/Engine/Plugins/Enterprise/GLTFExporter/Source/GLTFExporter/Private/Converters/GLTFCurveUtility.h b/Engine/Plugins/Enterprise/GLTFExporter/Source/GLTFExporter/Private/Converters/GLTFCurveUtility.h
deleted file mode 100644
diff --git a/Engine/Plugins/Enterprise/GLTFExporter/Source/GLTFExporter/Private/Converters/GLTFEpicLevelVariantSetsConverters.cpp b/Engine/Plugins/Enterprise/GLTFExporter/Source/GLTFExporter/Private/Converters/GLTFEpicLevelVariantSetsConverters.cpp
deleted file mode 100644
diff --git a/Engine/Plugins/Enterprise/GLTFExporter/Source/GLTFExporter/Private/Converters/GLTFKhrVariantConverters.cpp b/Engine/Plugins/Enterprise/GLTFExporter/Source/GLTFExporter/Private/Converters/GLTFKhrVariantConverters.cpp
--- a/Engine/Plugins/Enterprise/GLTFExporter/Source/GLTFExporter/Private/Converters/GLTFKhrVariantConverters.cpp
+++ b/Engine/Plugins/Enterprise/GLTFExporter/Source/GLTFExporter/Private/Converters/GLTFKhrVariantConverters.cpp
@@ -110,7 +110,7 @@ bool FGLTFKhrMaterialVariantConverter::TryParseMaterialProperty(FGLTFJsonPrimiti
 		return false;
 	}
 
-	const TArray<FCapturedPropSegment>& CapturedPropSegments = FGLTFVariantUtility::GetCapturedPropSegments(Property);
+	const TArray<FCapturedPropSegment>& CapturedPropSegments = Property->GetCapturedPropSegments();
 	const int32 NumPropSegments = CapturedPropSegments.Num();
 
 	if (NumPropSegments < 1)
diff --git a/Engine/Plugins/Enterprise/GLTFExporter/Source/GLTFExporter/Private/Converters/GLTFLightMapConverters.cpp b/Engine/Plugins/Enterprise/GLTFExporter/Source/GLTFExporter/Private/Converters/GLTFLightMapConverters.cpp
deleted file mode 100644
diff --git a/Engine/Plugins/Enterprise/GLTFExporter/Source/GLTFExporter/Private/Converters/GLTFMaterialProxyFactory.cpp b/Engine/Plugins/Enterprise/GLTFExporter/Source/GLTFExporter/Private/Converters/GLTFMaterialProxyFactory.cpp
--- a/Engine/Plugins/Enterprise/GLTFExporter/Source/GLTFExporter/Private/Converters/GLTFMaterialProxyFactory.cpp
+++ b/Engine/Plugins/Enterprise/GLTFExporter/Source/GLTFExporter/Private/Converters/GLTFMaterialProxyFactory.cpp
@@ -311,7 +311,6 @@ UGLTFExportOptions* FGLTFMaterialProxyFactory::CreateExportOptions(const UGLTFPr
 	UGLTFExportOptions* ExportOptions = NewObject<UGLTFExportOptions>();
 	ExportOptions->ResetToDefault();
 	ExportOptions->bExportProxyMaterials = false;
-	ExportOptions->bExportExtraBlendModes = true;
 	ExportOptions->BakeMaterialInputs = ProxyOptions->bBakeMaterialInputs ? EGLTFMaterialBakeMode::Simple : EGLTFMaterialBakeMode::Disabled;
 	ExportOptions->DefaultMaterialBakeSize = ProxyOptions->DefaultMaterialBakeSize;
 	ExportOptions->DefaultMaterialBakeFilter = ProxyOptions->DefaultMaterialBakeFilter;
diff --git a/Engine/Plugins/Enterprise/GLTFExporter/Source/GLTFExporter/Private/Converters/GLTFNodeConverters.cpp b/Engine/Plugins/Enterprise/GLTFExporter/Source/GLTFExporter/Private/Converters/GLTFNodeConverters.cpp
--- a/Engine/Plugins/Enterprise/GLTFExporter/Source/GLTFExporter/Private/Converters/GLTFNodeConverters.cpp
+++ b/Engine/Plugins/Enterprise/GLTFExporter/Source/GLTFExporter/Private/Converters/GLTFNodeConverters.cpp
@@ -3,7 +3,6 @@
 #include "Converters/GLTFNodeConverters.h"
 #include "Builders/GLTFContainerBuilder.h"
 #include "Utilities/GLTFCoreUtilities.h"
-#include "Converters/GLTFBlueprintUtility.h"
 #include "Converters/GLTFNameUtility.h"
 #include "LevelSequenceActor.h"
 #include "Camera/CameraComponent.h"
@@ -33,22 +32,7 @@ FGLTFJsonNode* FGLTFActorConverter::Convert(const AActor* Actor)
 
 	// TODO: process all components since any component can be attached to any other component in runtime
 
-	const FString BlueprintPath = FGLTFBlueprintUtility::GetClassPath(Actor);
-	if (FGLTFBlueprintUtility::IsSkySphere(BlueprintPath))
-	{
-		if (Builder.ExportOptions->bExportSkySpheres)
-		{
-			RootNode->SkySphere = Builder.AddUniqueSkySphere(Actor);
-		}
-	}
-	else if (FGLTFBlueprintUtility::IsHDRIBackdrop(BlueprintPath))
-	{
-		if (Builder.ExportOptions->bExportHDRIBackdrops)
-		{
-			RootNode->Backdrop = Builder.AddUniqueBackdrop(Actor);
-		}
-	}
-	else if (const ALevelSequenceActor* LevelSequenceActor = Cast<ALevelSequenceActor>(Actor))
+	if (const ALevelSequenceActor* LevelSequenceActor = Cast<ALevelSequenceActor>(Actor))
 	{
 		if (Builder.ExportOptions->bExportLevelSequences)
 		{
@@ -141,12 +125,6 @@ void FGLTFComponentConverter::ConvertComponentSpecialization(const USceneCompone
 		return;
 	}
 
-	const FString BlueprintPath = FGLTFBlueprintUtility::GetClassPath(Owner);
-	if (FGLTFBlueprintUtility::IsSkySphere(BlueprintPath) || FGLTFBlueprintUtility::IsHDRIBackdrop(BlueprintPath))
-	{
-		return;
-	}
-
 	if (Owner->IsA<ALevelSequenceActor>() || Owner->IsA<APawn>())
 	{
 		return;
@@ -155,11 +133,6 @@ void FGLTFComponentConverter::ConvertComponentSpecialization(const USceneCompone
 	if (const UStaticMeshComponent* StaticMeshComponent = Cast<UStaticMeshComponent>(SceneComponent))
 	{
 		Node->Mesh = Builder.AddUniqueMesh(StaticMeshComponent);
-
-		if (Builder.ExportOptions->bExportLightmaps)
-		{
-			Node->LightMap = Builder.AddUniqueLightMap(StaticMeshComponent);
-		}
 	}
 	else if (const USkeletalMeshComponent* SkeletalMeshComponent = Cast<USkeletalMeshComponent>(SceneComponent))
 	{
diff --git a/Engine/Plugins/Enterprise/GLTFExporter/Source/GLTFExporter/Private/Converters/GLTFSceneConverters.cpp b/Engine/Plugins/Enterprise/GLTFExporter/Source/GLTFExporter/Private/Converters/GLTFSceneConverters.cpp
--- a/Engine/Plugins/Enterprise/GLTFExporter/Source/GLTFExporter/Private/Converters/GLTFSceneConverters.cpp
+++ b/Engine/Plugins/Enterprise/GLTFExporter/Source/GLTFExporter/Private/Converters/GLTFSceneConverters.cpp
@@ -27,7 +27,7 @@ FGLTFJsonScene* FGLTFSceneConverter::Convert(const UWorld* World)
 
 			TArray<AActor*> LevelActors(Level->Actors); // iterate using copy, in case new temporary actors are added during conversion
 
-			if (Builder.ExportOptions->VariantSetsMode != EGLTFVariantSetsMode::None)
+			if (Builder.ExportOptions->ExportMaterialVariants != EGLTFMaterialVariantMode::None)
 			{
 				for (const AActor* Actor : LevelActors)
 				{
@@ -36,27 +36,13 @@ FGLTFJsonScene* FGLTFSceneConverter::Convert(const UWorld* World)
 					{
 						const ULevelVariantSets* LevelVariantSets = const_cast<ALevelVariantSetsActor*>(LevelVariantSetsActor)->GetLevelVariantSets(true);
 
-						if (Builder.ExportOptions->VariantSetsMode == EGLTFVariantSetsMode::Epic)
+						if (LevelVariantSets != nullptr)
 						{
-							if (LevelVariantSets != nullptr)
+							for (const UVariantSet* VariantSet: LevelVariantSets->GetVariantSets())
 							{
-								FGLTFJsonEpicLevelVariantSets* JsonEpicLevelVariantSets = Builder.AddUniqueEpicLevelVariantSets(LevelVariantSets);
-								if (JsonEpicLevelVariantSets != nullptr)
+								for (const UVariant* Variant: VariantSet->GetVariants())
 								{
-									Scene->EpicLevelVariantSets.Add(JsonEpicLevelVariantSets);
-								}
-							}
-						}
-						else if (Builder.ExportOptions->VariantSetsMode == EGLTFVariantSetsMode::Khronos)
-						{
-							if (LevelVariantSets != nullptr)
-							{
-								for (const UVariantSet* VariantSet: LevelVariantSets->GetVariantSets())
-								{
-									for (const UVariant* Variant: VariantSet->GetVariants())
-									{
-										Builder.AddUniqueKhrMaterialVariant(Variant);
-									}
+									Builder.AddUniqueKhrMaterialVariant(Variant);
 								}
 							}
 						}
diff --git a/Engine/Plugins/Enterprise/GLTFExporter/Source/GLTFExporter/Private/Converters/GLTFSkySphereConverters.cpp b/Engine/Plugins/Enterprise/GLTFExporter/Source/GLTFExporter/Private/Converters/GLTFSkySphereConverters.cpp
deleted file mode 100644
diff --git a/Engine/Plugins/Enterprise/GLTFExporter/Source/GLTFExporter/Private/Converters/GLTFTextureConverters.cpp b/Engine/Plugins/Enterprise/GLTFExporter/Source/GLTFExporter/Private/Converters/GLTFTextureConverters.cpp
--- a/Engine/Plugins/Enterprise/GLTFExporter/Source/GLTFExporter/Private/Converters/GLTFTextureConverters.cpp
+++ b/Engine/Plugins/Enterprise/GLTFExporter/Source/GLTFExporter/Private/Converters/GLTFTextureConverters.cpp
@@ -1,20 +1,7 @@
 // Copyright Epic Games, Inc. All Rights Reserved.
 
 #include "Converters/GLTFTextureConverters.h"
-#include "Converters/GLTFTextureUtility.h"
 #include "Tasks/GLTFDelayedTextureTasks.h"
-#include "Engine/Texture2D.h"
-#include "Engine/TextureCube.h"
-#include "Engine/TextureRenderTarget2D.h"
-#include "Engine/TextureRenderTargetCube.h"
-
-void FGLTFTexture2DConverter::Sanitize(const UTexture2D*& Texture2D, bool& bToSRGB)
-{
-	if (FGLTFTextureUtility::IsHDR(Texture2D) && Builder.GetTextureHDREncoding() != EGLTFJsonHDREncoding::None)
-	{
-		bToSRGB = false;
-	}
-}
 
 FGLTFJsonTexture* FGLTFTexture2DConverter::Convert(const UTexture2D* Texture2D, bool bToSRGB)
 {
@@ -28,34 +15,6 @@ FGLTFJsonTexture* FGLTFTexture2DConverter::Convert(const UTexture2D* Texture2D,
 	return JsonTexture;
 }
 
-void FGLTFTextureCubeConverter::Sanitize(const UTextureCube*& TextureCube, ECubeFace& CubeFace, bool& bToSRGB)
-{
-	if (FGLTFTextureUtility::IsHDR(TextureCube) && Builder.GetTextureHDREncoding() != EGLTFJsonHDREncoding::None)
-	{
-		bToSRGB = false;
-	}
-}
-
-FGLTFJsonTexture* FGLTFTextureCubeConverter::Convert(const UTextureCube* TextureCube, ECubeFace CubeFace, bool bToSRGB)
-{
-	if (Builder.ExportOptions->TextureImageFormat == EGLTFTextureImageFormat::None)
-	{
-		return nullptr;
-	}
-
-	FGLTFJsonTexture* JsonTexture = Builder.AddTexture();
-	Builder.ScheduleSlowTask<FGLTFDelayedTextureCubeTask>(Builder, TextureCube, CubeFace, bToSRGB, JsonTexture);
-	return JsonTexture;
-}
-
-void FGLTFTextureRenderTarget2DConverter::Sanitize(const UTextureRenderTarget2D*& RenderTarget2D, bool& bToSRGB)
-{
-	if (FGLTFTextureUtility::IsHDR(RenderTarget2D) && Builder.GetTextureHDREncoding() != EGLTFJsonHDREncoding::None)
-	{
-		bToSRGB = false;
-	}
-}
-
 FGLTFJsonTexture* FGLTFTextureRenderTarget2DConverter::Convert(const UTextureRenderTarget2D* RenderTarget2D, bool bToSRGB)
 {
 	if (Builder.ExportOptions->TextureImageFormat == EGLTFTextureImageFormat::None)
@@ -67,39 +26,3 @@ FGLTFJsonTexture* FGLTFTextureRenderTarget2DConverter::Convert(const UTextureRen
 	Builder.ScheduleSlowTask<FGLTFDelayedTextureRenderTarget2DTask>(Builder, RenderTarget2D, bToSRGB, JsonTexture);
 	return JsonTexture;
 }
-
-void FGLTFTextureRenderTargetCubeConverter::Sanitize(const UTextureRenderTargetCube*& RenderTargetCube, ECubeFace& CubeFace, bool& bToSRGB)
-{
-	if (FGLTFTextureUtility::IsHDR(RenderTargetCube) && Builder.GetTextureHDREncoding() != EGLTFJsonHDREncoding::None)
-	{
-		bToSRGB = false;
-	}
-}
-
-FGLTFJsonTexture* FGLTFTextureRenderTargetCubeConverter::Convert(const UTextureRenderTargetCube* RenderTargetCube, ECubeFace CubeFace, bool bToSRGB)
-{
-	if (Builder.ExportOptions->TextureImageFormat == EGLTFTextureImageFormat::None)
-	{
-		return nullptr;
-	}
-
-	FGLTFJsonTexture* JsonTexture = Builder.AddTexture();
-	Builder.ScheduleSlowTask<FGLTFDelayedTextureRenderTargetCubeTask>(Builder, RenderTargetCube, CubeFace, bToSRGB, JsonTexture);
-	return JsonTexture;
-}
-
-FGLTFJsonTexture* FGLTFTextureLightMapConverter::Convert(const ULightMapTexture2D* LightMap)
-{
-	if (Builder.ExportOptions->TextureImageFormat == EGLTFTextureImageFormat::None)
-	{
-		return nullptr;
-	}
-
-#if WITH_EDITOR
-	FGLTFJsonTexture* JsonTexture = Builder.AddTexture();
-	Builder.ScheduleSlowTask<FGLTFDelayedTextureLightMapTask>(Builder, LightMap, JsonTexture);
-	return JsonTexture;
-#else
-	return nullptr;
-#endif
-}
diff --git a/Engine/Plugins/Enterprise/GLTFExporter/Source/GLTFExporter/Private/Converters/GLTFTextureUtility.cpp b/Engine/Plugins/Enterprise/GLTFExporter/Source/GLTFExporter/Private/Converters/GLTFTextureUtility.cpp
--- a/Engine/Plugins/Enterprise/GLTFExporter/Source/GLTFExporter/Private/Converters/GLTFTextureUtility.cpp
+++ b/Engine/Plugins/Enterprise/GLTFExporter/Source/GLTFExporter/Private/Converters/GLTFTextureUtility.cpp
@@ -3,7 +3,6 @@
 #include "Converters/GLTFTextureUtility.h"
 #include "Converters/GLTFNormalMapPreview.h"
 #include "Converters/GLTFSimpleTexture2DPreview.h"
-#include "Converters/GLTFCubemapFacePreview.h"
 #include "CanvasItem.h"
 #include "CanvasTypes.h"
 #include "TextureCompiler.h"
@@ -230,98 +229,7 @@ bool FGLTFTextureUtility::RotateTexture(UTextureRenderTarget2D* OutTarget, const
 	return DrawTexture(OutTarget, InSource, InPosition, InSize, Transform);
 }
 
-UTexture2D* FGLTFTextureUtility::CreateTextureFromCubeFace(const UTextureCube* TextureCube, ECubeFace CubeFace)
-{
-	check(TextureCube != nullptr &&
-		TextureCube->TextureReference.IsInitialized() &&
-		TextureCube->TextureReference.TextureReferenceRHI.IsValid());
-
-	// Fully stream in the texture before drawing it.
-	const_cast<UTextureCube*>(TextureCube)->SetForceMipLevelsToBeResident(30.0f, true);
-	const_cast<UTextureCube*>(TextureCube)->WaitForStreaming();
-
-	const FIntPoint Size = GetInGameSize(TextureCube);
-	const bool bIsHDR = IsHDR(TextureCube);
-
-	UTextureRenderTarget2D* RenderTarget = CreateRenderTarget(Size, bIsHDR);
-	FTextureRenderTargetResource* RenderTargetResource = RenderTarget->GameThread_GetRenderTargetResource();
-
-	FCanvas Canvas(RenderTargetResource, nullptr, FGameTime::CreateDilated(0.0f, 0.0f, 0.0f, 0.0f), GMaxRHIFeatureLevel);
-	FCanvasTileItem TileItem(FVector2D::ZeroVector, TextureCube->GetResource(), FLinearColor::White);
-	TileItem.BatchedElementParameters = new FGLTFCubemapFacePreview(CubeFace);
-	TileItem.Draw(&Canvas);
-
-	Canvas.Flush_GameThread();
-	FlushRenderingCommands();
-	Canvas.SetRenderTarget_GameThread(nullptr);
-	FlushRenderingCommands();
-
-	UTexture2D* FaceTexture;
-	FTextureRenderTarget2DResource* Resource = static_cast<FTextureRenderTarget2DResource*>(RenderTarget->GetResource());
-	check(Resource != nullptr);
-
-	if (bIsHDR)
-	{
-		TArray<FFloat16Color> Pixels;
-		if (!Resource->ReadFloat16Pixels(Pixels, CubeFace))
-		{
-			return nullptr;
-		}
-
-		FaceTexture = CreateTransientTexture(Pixels.GetData(), Pixels.Num() * sizeof(FFloat16Color), Size, PF_FloatRGBA, false);
-		FaceTexture->CompressionSettings = TC_HDR;
-	}
-	else
-	{
-		TArray<FColor> Pixels;
-		if (!Resource->ReadPixels(Pixels, FReadSurfaceDataFlags(RCM_UNorm, CubeFace)))
-		{
-			return nullptr;
-		}
-
-		FaceTexture = CreateTransientTexture(Pixels.GetData(), Pixels.Num() * sizeof(FColor), Size, PF_B8G8R8A8, false);
-	}
-
-	return FaceTexture;
-}
-
-UTexture2D* FGLTFTextureUtility::CreateTextureFromCubeFace(const UTextureRenderTargetCube* RenderTargetCube, ECubeFace CubeFace)
-{
-	FTextureRenderTargetCubeResource* Resource = static_cast<FTextureRenderTargetCubeResource*>(const_cast<UTextureRenderTargetCube*>(RenderTargetCube)->GetResource());
-	if (Resource == nullptr)
-	{
-		return nullptr;
-	}
-
-	const FIntPoint Size(RenderTargetCube->SizeX, RenderTargetCube->SizeX);
-	UTexture2D* FaceTexture;
-
-	if (IsHDR(RenderTargetCube))
-	{
-		TArray<FFloat16Color> Pixels;
-		if (!Resource->ReadPixels(Pixels, FReadSurfaceDataFlags(RCM_UNorm, CubeFace)))
-		{
-			return nullptr;
-		}
-
-		FaceTexture = CreateTransientTexture(Pixels.GetData(), Pixels.Num() * sizeof(FFloat16Color), Size, PF_FloatRGBA, false);
-		FaceTexture->CompressionSettings = TC_HDR;
-	}
-	else
-	{
-		TArray<FColor> Pixels;
-		if (!Resource->ReadPixels(Pixels, FReadSurfaceDataFlags(RCM_UNorm, CubeFace)))
-		{
-			return nullptr;
-		}
-
-		FaceTexture = CreateTransientTexture(Pixels.GetData(), Pixels.Num() * sizeof(FColor), Size, PF_B8G8R8A8, false);
-	}
-
-	return FaceTexture;
-}
-
-bool FGLTFTextureUtility::ReadPixels(const UTextureRenderTarget2D* InRenderTarget, TArray<FColor>& OutPixels, EGLTFJsonHDREncoding Encoding)
+bool FGLTFTextureUtility::ReadPixels(const UTextureRenderTarget2D* InRenderTarget, TArray<FColor>& OutPixels)
 {
 	FTextureRenderTarget2DResource* Resource = static_cast<FTextureRenderTarget2DResource*>(const_cast<UTextureRenderTarget2D*>(InRenderTarget)->GetResource());
 	if (Resource == nullptr)
@@ -329,107 +237,9 @@ bool FGLTFTextureUtility::ReadPixels(const UTextureRenderTarget2D* InRenderTarge
 		return false;
 	}
 
-	if (Encoding == EGLTFJsonHDREncoding::None)
-	{
-		FReadSurfaceDataFlags ReadSurfaceDataFlags(RCM_UNorm, CubeFace_MAX);
-		ReadSurfaceDataFlags.SetLinearToGamma(false);
-		return Resource->ReadPixels(OutPixels, ReadSurfaceDataFlags);
-	}
-
-	TArray<FFloat16Color> HDRPixels;
-	if (!Resource->ReadFloat16Pixels(HDRPixels))
-	{
-		return false;
-	}
-
-	switch (Encoding)
-	{
-		case EGLTFJsonHDREncoding::RGBM:
-			EncodeRGBM(HDRPixels, OutPixels);
-			break;
-		case EGLTFJsonHDREncoding::RGBE:
-			EncodeRGBE(HDRPixels, OutPixels);
-			break;
-		default:
-			checkNoEntry();
-			break;
-	}
-
-	return true;
-}
-
-void FGLTFTextureUtility::EncodeRGBM(const TArray<FFloat16Color>& InPixels, TArray<FColor>& OutPixels, float MaxRange)
-{
-	OutPixels.AddUninitialized(InPixels.Num());
-
-	for (int32 Index = 0; Index < InPixels.Num(); ++Index)
-	{
-		const FFloat16Color& Color = InPixels[Index];
-		FLinearColor RGBM;
-
-		RGBM.R = FMath::Sqrt(Color.R.GetFloat());
-		RGBM.G = FMath::Sqrt(Color.G.GetFloat());
-		RGBM.B = FMath::Sqrt(Color.B.GetFloat());
-
-		RGBM.R /= MaxRange;
-		RGBM.G /= MaxRange;
-		RGBM.B /= MaxRange;
-
-		RGBM.A = FMath::Max(FMath::Max(RGBM.R, RGBM.G), FMath::Max(RGBM.B, 1.0f / 255.0f));
-		RGBM.A = FMath::CeilToFloat(RGBM.A * 255.0f) / 255.0f;
-
-		RGBM.R /= RGBM.A;
-		RGBM.G /= RGBM.A;
-		RGBM.B /= RGBM.A;
-
-		OutPixels[Index] = RGBM.ToFColor(false);
-	}
-}
-
-void FGLTFTextureUtility::EncodeRGBE(const TArray<FFloat16Color>& InPixels, TArray<FColor>& OutPixels)
-{
-	OutPixels.AddUninitialized(InPixels.Num());
-
-	for (int32 Index = 0; Index < InPixels.Num(); ++Index)
-	{
-		OutPixels[Index] = FLinearColor(InPixels[Index]).ToRGBE();
-	}
-}
-
-bool FGLTFTextureUtility::LoadPlatformData(UTexture2D* Texture)
-{
-	if (Texture->GetPlatformData() == nullptr || Texture->GetPlatformData()->Mips.Num() == 0)
-	{
-		return false;
-	}
-
-#if WITH_EDITOR
-	if (!Texture->GetPlatformData()->Mips[0].BulkData.IsBulkDataLoaded())
-	{
-		// TODO: is this correct handling?
-		Texture->ForceRebuildPlatformData();
-	}
-#endif
-
-	return Texture->GetPlatformData()->Mips[0].BulkData.IsBulkDataLoaded();
-}
-
-bool FGLTFTextureUtility::LoadPlatformData(UTextureCube* TextureCube)
-{
-	if (TextureCube->GetPlatformData() == nullptr || TextureCube->GetPlatformData()->Mips.Num() == 0)
-	{
-		return false;
-	}
-
-#if WITH_EDITOR
-	if (!TextureCube->GetPlatformData()->Mips[0].BulkData.IsBulkDataLoaded())
-	{
-		// TODO: is this correct handling?
-		TextureCube->ForceRebuildPlatformData();
-	}
-#endif
-
-	return TextureCube->GetPlatformData()->Mips[0].BulkData.IsBulkDataLoaded();
+	FReadSurfaceDataFlags ReadSurfaceDataFlags(RCM_UNorm, CubeFace_MAX);
+	ReadSurfaceDataFlags.SetLinearToGamma(false);
+	return Resource->ReadPixels(OutPixels, ReadSurfaceDataFlags);
 }
 
 void FGLTFTextureUtility::FlipGreenChannel(TArray<FColor>& Pixels)
diff --git a/Engine/Plugins/Enterprise/GLTFExporter/Source/GLTFExporter/Private/Converters/GLTFTextureUtility.h b/Engine/Plugins/Enterprise/GLTFExporter/Source/GLTFExporter/Private/Converters/GLTFTextureUtility.h
--- a/Engine/Plugins/Enterprise/GLTFExporter/Source/GLTFExporter/Private/Converters/GLTFTextureUtility.h
+++ b/Engine/Plugins/Enterprise/GLTFExporter/Source/GLTFExporter/Private/Converters/GLTFTextureUtility.h
@@ -38,17 +38,7 @@ struct FGLTFTextureUtility
 	static bool DrawTexture(UTextureRenderTarget2D* OutTarget, const UTexture2D* InSource, const FVector2D& InPosition, const FVector2D& InSize, const FMatrix& InTransform = FMatrix::Identity);
 	static bool RotateTexture(UTextureRenderTarget2D* OutTarget, const UTexture2D* InSource, const FVector2D& InPosition, const FVector2D& InSize, float InDegrees);
 
-	static UTexture2D* CreateTextureFromCubeFace(const UTextureCube* TextureCube, ECubeFace CubeFace);
-	static UTexture2D* CreateTextureFromCubeFace(const UTextureRenderTargetCube* RenderTargetCube, ECubeFace CubeFace);
-
-	static bool ReadPixels(const UTextureRenderTarget2D* InRenderTarget, TArray<FColor>& OutPixels, EGLTFJsonHDREncoding Encoding);
-
-	static void EncodeRGBM(const TArray<FFloat16Color>& InPixels, TArray<FColor>& OutPixels, float MaxRange = 8);
-	static void EncodeRGBE(const TArray<FFloat16Color>& InPixels, TArray<FColor>& OutPixels);
-
-	// TODO: maybe use template specialization to avoid the need for duplicated functions
-	static bool LoadPlatformData(UTexture2D* Texture);
-	static bool LoadPlatformData(UTextureCube* TextureCube);
+	static bool ReadPixels(const UTextureRenderTarget2D* InRenderTarget, TArray<FColor>& OutPixels);
 
 	static void FlipGreenChannel(TArray<FColor>& Pixels);
 	static void TransformColorSpace(TArray<FColor>& Pixels, bool bFromSRGB, bool bToSRGB);
diff --git a/Engine/Plugins/Enterprise/GLTFExporter/Source/GLTFExporter/Private/Converters/GLTFVariantUtility.cpp b/Engine/Plugins/Enterprise/GLTFExporter/Source/GLTFExporter/Private/Converters/GLTFVariantUtility.cpp
--- a/Engine/Plugins/Enterprise/GLTFExporter/Source/GLTFExporter/Private/Converters/GLTFVariantUtility.cpp
+++ b/Engine/Plugins/Enterprise/GLTFExporter/Source/GLTFExporter/Private/Converters/GLTFVariantUtility.cpp
@@ -8,26 +8,6 @@
 #include "VariantObjectBinding.h"
 #include "PropertyValue.h"
 
-const TArray<FCapturedPropSegment>& FGLTFVariantUtility::GetCapturedPropSegments(const UPropertyValue* Property)
-{
-	return Property->GetCapturedPropSegments();
-}
-
-bool FGLTFVariantUtility::IsStaticMeshProperty(const UPropertyValue* Property)
-{
-	return Property->GetPropertyName() == TEXT("StaticMesh") && Property->GetPropertyClass()->IsChildOf(FObjectPropertyBase::StaticClass());
-}
-
-bool FGLTFVariantUtility::IsSkeletalMeshProperty(const UPropertyValue* Property)
-{
-	return Property->GetPropertyName() == TEXT("SkeletalMesh") && Property->GetPropertyClass()->IsChildOf(FObjectPropertyBase::StaticClass());
-}
-
-bool FGLTFVariantUtility::IsVisibleProperty(const UPropertyValue* Property)
-{
-	return Property->GetPropertyName() == TEXT("bVisible") && Property->GetPropertyClass() == FBoolProperty::StaticClass();
-}
-
 bool FGLTFVariantUtility::TryGetPropertyValue(UPropertyValue* Property, void* OutData, uint32 OutSize)
 {
 	if (Property == nullptr || !Property->HasRecordedData())
diff --git a/Engine/Plugins/Enterprise/GLTFExporter/Source/GLTFExporter/Private/Converters/GLTFVariantUtility.h b/Engine/Plugins/Enterprise/GLTFExporter/Source/GLTFExporter/Private/Converters/GLTFVariantUtility.h
--- a/Engine/Plugins/Enterprise/GLTFExporter/Source/GLTFExporter/Private/Converters/GLTFVariantUtility.h
+++ b/Engine/Plugins/Enterprise/GLTFExporter/Source/GLTFExporter/Private/Converters/GLTFVariantUtility.h
@@ -16,12 +16,6 @@ class UPropertyValue;
 
 struct FGLTFVariantUtility
 {
-	static const TArray<FCapturedPropSegment>& GetCapturedPropSegments(const UPropertyValue* Property);
-
-	static bool IsStaticMeshProperty(const UPropertyValue* Property);
-	static bool IsSkeletalMeshProperty(const UPropertyValue* Property);
-	static bool IsVisibleProperty(const UPropertyValue* Property);
-
 	template <typename T>
 	static bool TryGetPropertyValue(UPropertyValue* Property, T& OutValue)
 	{
diff --git a/Engine/Plugins/Enterprise/GLTFExporter/Source/GLTFExporter/Private/Exporters/GLTFLevelVariantSetsExporter.cpp b/Engine/Plugins/Enterprise/GLTFExporter/Source/GLTFExporter/Private/Exporters/GLTFLevelVariantSetsExporter.cpp
--- a/Engine/Plugins/Enterprise/GLTFExporter/Source/GLTFExporter/Private/Exporters/GLTFLevelVariantSetsExporter.cpp
+++ b/Engine/Plugins/Enterprise/GLTFExporter/Source/GLTFExporter/Private/Exporters/GLTFLevelVariantSetsExporter.cpp
@@ -16,10 +16,10 @@ bool UGLTFLevelVariantSetsExporter::AddObject(FGLTFContainerBuilder& Builder, co
 {
 	const ULevelVariantSets* LevelVariantSets = CastChecked<ULevelVariantSets>(Object);
 
-	if (Builder.ExportOptions->VariantSetsMode == EGLTFVariantSetsMode::None)
+	if (Builder.ExportOptions->ExportMaterialVariants == EGLTFMaterialVariantMode::None)
 	{
 		Builder.LogError(
-			FString::Printf(TEXT("Failed to export level variant sets %s because variant sets are disabled by export options"),
+			FString::Printf(TEXT("Failed to export level variant sets %s because material variants are disabled by export options"),
 			*LevelVariantSets->GetName()));
 		return false;
 	}
@@ -52,28 +52,12 @@ bool UGLTFLevelVariantSetsExporter::AddObject(FGLTFContainerBuilder& Builder, co
 		return false;
 	}
 
-	if (Builder.ExportOptions->VariantSetsMode == EGLTFVariantSetsMode::Epic)
+	if (Builder.GetRoot().KhrMaterialVariants.Num() < 1)
 	{
-		FGLTFJsonEpicLevelVariantSets* EpicLevelVariantSets = Builder.AddUniqueEpicLevelVariantSets(LevelVariantSets);
-		if (EpicLevelVariantSets == nullptr)
-		{
-			Builder.LogError(
-				FString::Printf(TEXT("Failed to export level variant sets %s"),
-				*LevelVariantSets->GetName()));
-			return false;
-		}
-
-		Scene->EpicLevelVariantSets.AddUnique(EpicLevelVariantSets);
-	}
-	else if (Builder.ExportOptions->VariantSetsMode == EGLTFVariantSetsMode::Khronos)
-	{
-		if (Builder.GetRoot().KhrMaterialVariants.Num() < 1)
-		{
-			Builder.LogError(
-				FString::Printf(TEXT("Failed to export any supported variant from level variant sets %s"),
-				*LevelVariantSets->GetName()));
-			return false;
-		}
+		Builder.LogError(
+			FString::Printf(TEXT("Failed to export any supported variant from level variant sets %s"),
+			*LevelVariantSets->GetName()));
+		return false;
 	}
 
 	Builder.DefaultScene = Scene;
diff --git a/Engine/Plugins/Enterprise/GLTFExporter/Source/GLTFExporter/Private/Json/GLTFJsonAnimation.cpp b/Engine/Plugins/Enterprise/GLTFExporter/Source/GLTFExporter/Private/Json/GLTFJsonAnimation.cpp
--- a/Engine/Plugins/Enterprise/GLTFExporter/Source/GLTFExporter/Private/Json/GLTFJsonAnimation.cpp
+++ b/Engine/Plugins/Enterprise/GLTFExporter/Source/GLTFExporter/Private/Json/GLTFJsonAnimation.cpp
@@ -37,11 +37,4 @@ void FGLTFJsonAnimation::WriteObject(IGLTFJsonWriter& Writer) const
 
 	Writer.Write(TEXT("channels"), Channels);
 	Writer.Write(TEXT("samplers"), Samplers);
-
-	if (Playback != FGLTFJsonAnimationPlayback())
-	{
-		Writer.StartExtensions();
-		Writer.Write(EGLTFJsonExtension::EPIC_AnimationPlayback, Playback);
-		Writer.EndExtensions();
-	}
 }
diff --git a/Engine/Plugins/Enterprise/GLTFExporter/Source/GLTFExporter/Private/Json/GLTFJsonAnimationPlayback.cpp b/Engine/Plugins/Enterprise/GLTFExporter/Source/GLTFExporter/Private/Json/GLTFJsonAnimationPlayback.cpp
deleted file mode 100644
diff --git a/Engine/Plugins/Enterprise/GLTFExporter/Source/GLTFExporter/Private/Json/GLTFJsonBackdrop.cpp b/Engine/Plugins/Enterprise/GLTFExporter/Source/GLTFExporter/Private/Json/GLTFJsonBackdrop.cpp
deleted file mode 100644
diff --git a/Engine/Plugins/Enterprise/GLTFExporter/Source/GLTFExporter/Private/Json/GLTFJsonEpicLevelVariantSets.cpp b/Engine/Plugins/Enterprise/GLTFExporter/Source/GLTFExporter/Private/Json/GLTFJsonEpicLevelVariantSets.cpp
deleted file mode 100644
diff --git a/Engine/Plugins/Enterprise/GLTFExporter/Source/GLTFExporter/Private/Json/GLTFJsonLightMap.cpp b/Engine/Plugins/Enterprise/GLTFExporter/Source/GLTFExporter/Private/Json/GLTFJsonLightMap.cpp
deleted file mode 100644
diff --git a/Engine/Plugins/Enterprise/GLTFExporter/Source/GLTFExporter/Private/Json/GLTFJsonMaterial.cpp b/Engine/Plugins/Enterprise/GLTFExporter/Source/GLTFExporter/Private/Json/GLTFJsonMaterial.cpp
--- a/Engine/Plugins/Enterprise/GLTFExporter/Source/GLTFExporter/Private/Json/GLTFJsonMaterial.cpp
+++ b/Engine/Plugins/Enterprise/GLTFExporter/Source/GLTFExporter/Private/Json/GLTFJsonMaterial.cpp
@@ -153,17 +153,10 @@ void FGLTFJsonMaterial::WriteObject(IGLTFJsonWriter& Writer) const
 		Writer.Write(TEXT("doubleSided"), DoubleSided);
 	}
 
-	if (BlendMode != EGLTFJsonBlendMode::None || ShadingModel == EGLTFJsonShadingModel::Unlit || ShadingModel == EGLTFJsonShadingModel::ClearCoat)
+	if (ShadingModel == EGLTFJsonShadingModel::Unlit || ShadingModel == EGLTFJsonShadingModel::ClearCoat)
 	{
 		Writer.StartExtensions();
 
-		if (BlendMode != EGLTFJsonBlendMode::None)
-		{
-			Writer.StartExtension(EGLTFJsonExtension::EPIC_BlendModes);
-			Writer.Write(TEXT("blendMode"), BlendMode);
-			Writer.EndExtension();
-		}
-
 		if (ShadingModel == EGLTFJsonShadingModel::Unlit)
 		{
 			Writer.StartExtension(EGLTFJsonExtension::KHR_MaterialsUnlit);
diff --git a/Engine/Plugins/Enterprise/GLTFExporter/Source/GLTFExporter/Private/Json/GLTFJsonNode.cpp b/Engine/Plugins/Enterprise/GLTFExporter/Source/GLTFExporter/Private/Json/GLTFJsonNode.cpp
--- a/Engine/Plugins/Enterprise/GLTFExporter/Source/GLTFExporter/Private/Json/GLTFJsonNode.cpp
+++ b/Engine/Plugins/Enterprise/GLTFExporter/Source/GLTFExporter/Private/Json/GLTFJsonNode.cpp
@@ -4,10 +4,7 @@
 #include "Json/GLTFJsonCamera.h"
 #include "Json/GLTFJsonSkin.h"
 #include "Json/GLTFJsonMesh.h"
-#include "Json/GLTFJsonBackdrop.h"
 #include "Json/GLTFJsonLight.h"
-#include "Json/GLTFJsonLightMap.h"
-#include "Json/GLTFJsonSkySphere.h"
 
 void FGLTFJsonNode::WriteObject(IGLTFJsonWriter& Writer) const
 {
@@ -46,37 +43,13 @@ void FGLTFJsonNode::WriteObject(IGLTFJsonWriter& Writer) const
 		Writer.Write(TEXT("mesh"), Mesh);
 	}
 
-	if (Backdrop != nullptr || Light != nullptr || LightMap != nullptr || SkySphere != nullptr)
+	if (Light != nullptr)
 	{
 		Writer.StartExtensions();
 
-		if (Backdrop != nullptr)
-		{
-			Writer.StartExtension(EGLTFJsonExtension::EPIC_HDRIBackdrops);
-			Writer.Write(TEXT("backdrop"), Backdrop);
-			Writer.EndExtension();
-		}
-
-		if (Light != nullptr)
-		{
-			Writer.StartExtension(EGLTFJsonExtension::KHR_LightsPunctual);
-			Writer.Write(TEXT("light"), Light);
-			Writer.EndExtension();
-		}
-
-		if (LightMap != nullptr)
-		{
-			Writer.StartExtension(EGLTFJsonExtension::EPIC_LightmapTextures);
-			Writer.Write(TEXT("lightmap"), LightMap);
-			Writer.EndExtension();
-		}
-
-		if (SkySphere != nullptr)
-		{
-			Writer.StartExtension(EGLTFJsonExtension::EPIC_SkySpheres);
-			Writer.Write(TEXT("skySphere"), SkySphere);
-			Writer.EndExtension();
-		}
+		Writer.StartExtension(EGLTFJsonExtension::KHR_LightsPunctual);
+		Writer.Write(TEXT("light"), Light);
+		Writer.EndExtension();
 
 		Writer.EndExtensions();
 	}
diff --git a/Engine/Plugins/Enterprise/GLTFExporter/Source/GLTFExporter/Private/Json/GLTFJsonRoot.cpp b/Engine/Plugins/Enterprise/GLTFExporter/Source/GLTFExporter/Private/Json/GLTFJsonRoot.cpp
--- a/Engine/Plugins/Enterprise/GLTFExporter/Source/GLTFExporter/Private/Json/GLTFJsonRoot.cpp
+++ b/Engine/Plugins/Enterprise/GLTFExporter/Source/GLTFExporter/Private/Json/GLTFJsonRoot.cpp
@@ -25,24 +25,10 @@ void FGLTFJsonRoot::WriteObject(IGLTFJsonWriter& Writer) const
 	if (Skins.Num() > 0) Writer.Write(TEXT("skins"), Skins);
 	if (Textures.Num() > 0) Writer.Write(TEXT("textures"), Textures);
 
-	if (Backdrops.Num() > 0 || Lights.Num() > 0 || LightMaps.Num() > 0 || SkySpheres.Num() > 0 || EpicLevelVariantSets.Num() > 0 || KhrMaterialVariants.Num() > 0)
+	if (Lights.Num() > 0 || KhrMaterialVariants.Num() > 0)
 	{
 		Writer.StartExtensions();
 
-		if (Backdrops.Num() > 0)
-		{
-			Writer.StartExtension(EGLTFJsonExtension::EPIC_HDRIBackdrops);
-			Writer.Write(TEXT("backdrops"), Backdrops);
-			Writer.EndExtension();
-		}
-
-		if (EpicLevelVariantSets.Num() > 0)
-		{
-			Writer.StartExtension(EGLTFJsonExtension::EPIC_LevelVariantSets);
-			Writer.Write(TEXT("levelVariantSets"), EpicLevelVariantSets);
-			Writer.EndExtension();
-		}
-
 		if (KhrMaterialVariants.Num() > 0)
 		{
 			Writer.StartExtension(EGLTFJsonExtension::KHR_MaterialsVariants);
@@ -57,20 +43,6 @@ void FGLTFJsonRoot::WriteObject(IGLTFJsonWriter& Writer) const
 			Writer.EndExtension();
 		}
 
-		if (LightMaps.Num() > 0)
-		{
-			Writer.StartExtension(EGLTFJsonExtension::EPIC_LightmapTextures);
-			Writer.Write(TEXT("lightmaps"), LightMaps);
-			Writer.EndExtension();
-		}
-
-		if (SkySpheres.Num() > 0)
-		{
-			Writer.StartExtension(EGLTFJsonExtension::EPIC_SkySpheres);
-			Writer.Write(TEXT("skySpheres"), SkySpheres);
-			Writer.EndExtension();
-		}
-
 		Writer.EndExtensions();
 	}
 
diff --git a/Engine/Plugins/Enterprise/GLTFExporter/Source/GLTFExporter/Private/Json/GLTFJsonScene.cpp b/Engine/Plugins/Enterprise/GLTFExporter/Source/GLTFExporter/Private/Json/GLTFJsonScene.cpp
--- a/Engine/Plugins/Enterprise/GLTFExporter/Source/GLTFExporter/Private/Json/GLTFJsonScene.cpp
+++ b/Engine/Plugins/Enterprise/GLTFExporter/Source/GLTFExporter/Private/Json/GLTFJsonScene.cpp
@@ -2,7 +2,6 @@
 
 #include "Json/GLTFJsonScene.h"
 #include "Json/GLTFJsonNode.h"
-#include "Json/GLTFJsonEpicLevelVariantSets.h"
 
 void FGLTFJsonScene::WriteObject(IGLTFJsonWriter& Writer) const
 {
@@ -15,15 +14,4 @@ void FGLTFJsonScene::WriteObject(IGLTFJsonWriter& Writer) const
 	{
 		Writer.Write(TEXT("nodes"), Nodes);
 	}
-
-	if (EpicLevelVariantSets.Num() > 0)
-	{
-		Writer.StartExtensions();
-
-		Writer.StartExtension(EGLTFJsonExtension::EPIC_LevelVariantSets);
-		Writer.Write(TEXT("levelVariantSets"), EpicLevelVariantSets);
-		Writer.EndExtension();
-
-		Writer.EndExtensions();
-	}
 }
diff --git a/Engine/Plugins/Enterprise/GLTFExporter/Source/GLTFExporter/Private/Json/GLTFJsonSkySphere.cpp b/Engine/Plugins/Enterprise/GLTFExporter/Source/GLTFExporter/Private/Json/GLTFJsonSkySphere.cpp
deleted file mode 100644
diff --git a/Engine/Plugins/Enterprise/GLTFExporter/Source/GLTFExporter/Private/Json/GLTFJsonTexture.cpp b/Engine/Plugins/Enterprise/GLTFExporter/Source/GLTFExporter/Private/Json/GLTFJsonTexture.cpp
--- a/Engine/Plugins/Enterprise/GLTFExporter/Source/GLTFExporter/Private/Json/GLTFJsonTexture.cpp
+++ b/Engine/Plugins/Enterprise/GLTFExporter/Source/GLTFExporter/Private/Json/GLTFJsonTexture.cpp
@@ -20,15 +20,4 @@ void FGLTFJsonTexture::WriteObject(IGLTFJsonWriter& Writer) const
 	{
 		Writer.Write(TEXT("source"), Source);
 	}
-
-	if (Encoding != EGLTFJsonHDREncoding::None)
-	{
-		Writer.StartExtensions();
-
-		Writer.StartExtension(EGLTFJsonExtension::EPIC_TextureHDREncoding);
-		Writer.Write(TEXT("encoding"), Encoding);
-		Writer.EndExtension();
-
-		Writer.EndExtensions();
-	}
 }
diff --git a/Engine/Plugins/Enterprise/GLTFExporter/Source/GLTFExporter/Private/Options/GLTFExportOptions.cpp b/Engine/Plugins/Enterprise/GLTFExporter/Source/GLTFExporter/Private/Options/GLTFExportOptions.cpp
--- a/Engine/Plugins/Enterprise/GLTFExporter/Source/GLTFExporter/Private/Options/GLTFExportOptions.cpp
+++ b/Engine/Plugins/Enterprise/GLTFExporter/Source/GLTFExporter/Private/Options/GLTFExportOptions.cpp
@@ -18,7 +18,6 @@ void UGLTFExportOptions::ResetToDefault()
 	bExportProxyMaterials = true;
 	bExportUnlitMaterials = true;
 	bExportClearCoatMaterials = true;
-	bExportExtraBlendModes = false;
 	BakeMaterialInputs = EGLTFMaterialBakeMode::UseMeshData;
 	DefaultMaterialBakeSize = EGLTFMaterialBakeSizePOT::POT_1024;
 	DefaultMaterialBakeFilter = TF_Trilinear;
@@ -29,21 +28,13 @@ void UGLTFExportOptions::ResetToDefault()
 	bUseMeshQuantization = false;
 	bExportLevelSequences = true;
 	bExportAnimationSequences = true;
-	bExportPlaybackSettings = false;
 	TextureImageFormat = EGLTFTextureImageFormat::PNG;
 	TextureImageQuality = 0;
 	NoLossyImageFormatFor = static_cast<int32>(EGLTFTextureType::All);
 	bExportTextureTransforms = true;
-	bExportLightmaps = false;
-	TextureHDREncoding = EGLTFTextureHDREncoding::None;
 	bAdjustNormalmaps = true;
 	bExportHiddenInGame = false;
 	ExportLights = static_cast<int32>(EGLTFSceneMobility::Stationary | EGLTFSceneMobility::Movable);
 	bExportCameras = true;
-	bExportHDRIBackdrops = false;
-	bExportSkySpheres = false;
-	VariantSetsMode = EGLTFVariantSetsMode::None;
 	ExportMaterialVariants = EGLTFMaterialVariantMode::UseMeshData;
-	bExportMeshVariants = true;
-	bExportVisibilityVariants = true;
 }
diff --git a/Engine/Plugins/Enterprise/GLTFExporter/Source/GLTFExporter/Private/Tasks/GLTFDelayedMaterialTasks.cpp b/Engine/Plugins/Enterprise/GLTFExporter/Source/GLTFExporter/Private/Tasks/GLTFDelayedMaterialTasks.cpp
--- a/Engine/Plugins/Enterprise/GLTFExporter/Source/GLTFExporter/Private/Tasks/GLTFDelayedMaterialTasks.cpp
+++ b/Engine/Plugins/Enterprise/GLTFExporter/Source/GLTFExporter/Private/Tasks/GLTFDelayedMaterialTasks.cpp
@@ -60,7 +60,7 @@ void FGLTFDelayedMaterialTask::Process()
 	JsonMaterial->DoubleSided = Material->IsTwoSided();
 
 	ConvertShadingModel(JsonMaterial->ShadingModel);
-	ConvertAlphaMode(JsonMaterial->AlphaMode, JsonMaterial->BlendMode);
+	ConvertAlphaMode(JsonMaterial->AlphaMode);
 
 	if (FGLTFProxyMaterialUtilities::IsProxyMaterial(BaseMaterial))
 	{
@@ -435,7 +435,7 @@ void FGLTFDelayedMaterialTask::ConvertShadingModel(EGLTFJsonShadingModel& OutSha
 	}
 }
 
-void FGLTFDelayedMaterialTask::ConvertAlphaMode(EGLTFJsonAlphaMode& OutAlphaMode, EGLTFJsonBlendMode& OutBlendMode) const
+void FGLTFDelayedMaterialTask::ConvertAlphaMode(EGLTFJsonAlphaMode& OutAlphaMode) const
 {
 	const EBlendMode BlendMode = Material->GetBlendMode();
 
@@ -443,7 +443,6 @@ void FGLTFDelayedMaterialTask::ConvertAlphaMode(EGLTFJsonAlphaMode& OutAlphaMode
 	if (OutAlphaMode == EGLTFJsonAlphaMode::None)
 	{
 		OutAlphaMode = EGLTFJsonAlphaMode::Blend;
-		OutBlendMode = EGLTFJsonBlendMode::None;
 
 		Builder.LogWarning(FString::Printf(
 			TEXT("Unsupported blend mode (%s) in material %s, will export as %s"),
@@ -452,21 +451,6 @@ void FGLTFDelayedMaterialTask::ConvertAlphaMode(EGLTFJsonAlphaMode& OutAlphaMode
 			*FGLTFNameUtility::GetName(BLEND_Translucent)));
 		return;
 	}
-
-	if (OutAlphaMode == EGLTFJsonAlphaMode::Blend)
-	{
-		OutBlendMode = FGLTFCoreUtilities::ConvertBlendMode(BlendMode);
-		if (OutBlendMode != EGLTFJsonBlendMode::None && !Builder.ExportOptions->bExportExtraBlendModes)
-		{
-			OutBlendMode = EGLTFJsonBlendMode::None;
-
-			Builder.LogWarning(FString::Printf(
-				TEXT("Extra blend mode (%s) in material %s disabled by export options, will export as %s"),
-				*FGLTFNameUtility::GetName(BlendMode),
-				*Material->GetName(),
-				*FGLTFNameUtility::GetName(BLEND_Translucent)));
-		}
-	}
 }
 
 #if WITH_EDITOR
diff --git a/Engine/Plugins/Enterprise/GLTFExporter/Source/GLTFExporter/Private/Tasks/GLTFDelayedMaterialTasks.h b/Engine/Plugins/Enterprise/GLTFExporter/Source/GLTFExporter/Private/Tasks/GLTFDelayedMaterialTasks.h
--- a/Engine/Plugins/Enterprise/GLTFExporter/Source/GLTFExporter/Private/Tasks/GLTFDelayedMaterialTasks.h
+++ b/Engine/Plugins/Enterprise/GLTFExporter/Source/GLTFExporter/Private/Tasks/GLTFDelayedMaterialTasks.h
@@ -56,7 +56,7 @@ private:
 
 	EMaterialShadingModel GetShadingModel() const;
 	void ConvertShadingModel(EGLTFJsonShadingModel& OutShadingModel) const;
-	void ConvertAlphaMode(EGLTFJsonAlphaMode& OutAlphaMode, EGLTFJsonBlendMode& OutBlendMode) const;
+	void ConvertAlphaMode(EGLTFJsonAlphaMode& OutAlphaMode) const;
 
 #if WITH_EDITOR
 	TSet<FGLTFMaterialPropertyEx> MeshDataBakedProperties;
diff --git a/Engine/Plugins/Enterprise/GLTFExporter/Source/GLTFExporter/Private/Tasks/GLTFDelayedTextureTasks.cpp b/Engine/Plugins/Enterprise/GLTFExporter/Source/GLTFExporter/Private/Tasks/GLTFDelayedTextureTasks.cpp
--- a/Engine/Plugins/Enterprise/GLTFExporter/Source/GLTFExporter/Private/Tasks/GLTFDelayedTextureTasks.cpp
+++ b/Engine/Plugins/Enterprise/GLTFExporter/Source/GLTFExporter/Private/Tasks/GLTFDelayedTextureTasks.cpp
@@ -34,13 +34,8 @@ void FGLTFDelayedTexture2DTask::Process()
 	// TODO: preserve maximum image quality (avoid compression artifacts) by copying source data (and adjustments) to a temp texture
 	FGLTFTextureUtility::DrawTexture(RenderTarget, Texture2D, FVector2D::ZeroVector, Size);
 
-	if (bIsHDR)
-	{
-		JsonTexture->Encoding = Builder.GetTextureHDREncoding();
-	}
-
 	TGLTFSharedArray<FColor> Pixels;
-	if (!FGLTFTextureUtility::ReadPixels(RenderTarget, *Pixels, JsonTexture->Encoding))
+	if (!FGLTFTextureUtility::ReadPixels(RenderTarget, *Pixels))
 	{
 		Builder.LogWarning(FString::Printf(TEXT("Failed to read pixels for 2D texture %s"), *JsonTexture->Name));
 		return;
@@ -63,64 +58,6 @@ void FGLTFDelayedTexture2DTask::Process()
 	JsonTexture->Sampler = Builder.AddUniqueSampler(Texture2D);
 }
 
-FString FGLTFDelayedTextureCubeTask::GetName()
-{
-	return TextureCube->GetName();
-}
-
-void FGLTFDelayedTextureCubeTask::Process()
-{
-	FGLTFTextureUtility::FullyLoad(TextureCube);
-	JsonTexture->Name = TextureCube->GetName() + TEXT("_") + FGLTFJsonUtilities::GetValue(FGLTFCoreUtilities::ConvertCubeFace(CubeFace));
-
-	const bool bFromSRGB = TextureCube->SRGB;
-	if (bFromSRGB != bToSRGB)
-	{
-		JsonTexture->Name += bToSRGB ? TEXT("_sRGB") : TEXT("_Linear");
-	}
-
-	// TODO: add optimized "happy path" if cube face doesn't need rotation and has suitable pixel format
-
-	// TODO: preserve maximum image quality (avoid compression artifacts) by copying source data (and adjustments) to a temp texture
-	const UTexture2D* FaceTexture = FGLTFTextureUtility::CreateTextureFromCubeFace(TextureCube, CubeFace);
-	if (FaceTexture == nullptr)
-	{
-		Builder.LogWarning(FString::Printf(TEXT("Failed to extract cube face %d for cubemap texture %s"), CubeFace, *TextureCube->GetName()));
-		return;
-	}
-
-	const bool bIsHDR = FGLTFTextureUtility::IsHDR(TextureCube);
-	const FIntPoint Size = { TextureCube->GetSizeX(), TextureCube->GetSizeY() };
-	UTextureRenderTarget2D* RenderTarget = FGLTFTextureUtility::CreateRenderTarget(Size, bIsHDR);
-
-	const float FaceRotation = FGLTFTextureUtility::GetCubeFaceRotation(CubeFace);
-	FGLTFTextureUtility::RotateTexture(RenderTarget, FaceTexture, FVector2D::ZeroVector, Size, FaceRotation);
-
-	if (bIsHDR)
-	{
-		JsonTexture->Encoding = Builder.GetTextureHDREncoding();
-		if (JsonTexture->Encoding == EGLTFJsonHDREncoding::None)
-		{
-			bToSRGB = true;
-		}
-	}
-
-	TGLTFSharedArray<FColor> Pixels;
-	if (!FGLTFTextureUtility::ReadPixels(RenderTarget, *Pixels, JsonTexture->Encoding))
-	{
-		Builder.LogWarning(FString::Printf(TEXT("Failed to read pixels (cube face %d) for cubemap texture %s"), CubeFace, *TextureCube->GetName()));
-		return;
-	}
-
-	FGLTFTextureUtility::TransformColorSpace(*Pixels, bFromSRGB, bToSRGB);
-
-	const bool bIgnoreAlpha = FGLTFTextureUtility::IsAlphaless(TextureCube->GetPixelFormat());
-	const EGLTFTextureType Type = bIsHDR ? EGLTFTextureType::HDR : EGLTFTextureType::None;
-
-	JsonTexture->Source = Builder.AddUniqueImage(Pixels, Size, bIgnoreAlpha, Type, JsonTexture->Name);
-	JsonTexture->Sampler = Builder.AddUniqueSampler(TextureCube);
-}
-
 FString FGLTFDelayedTextureRenderTarget2DTask::GetName()
 {
 	return RenderTarget2D->GetName();
@@ -140,13 +77,8 @@ void FGLTFDelayedTextureRenderTarget2DTask::Process()
 	const bool bIsHDR = FGLTFTextureUtility::IsHDR(RenderTarget2D);
 	const FIntPoint Size = { RenderTarget2D->SizeX, RenderTarget2D->SizeY };
 
-	if (bIsHDR)
-	{
-		JsonTexture->Encoding = Builder.GetTextureHDREncoding();
-	}
-
 	TGLTFSharedArray<FColor> Pixels;
-	if (!FGLTFTextureUtility::ReadPixels(RenderTarget2D, *Pixels, JsonTexture->Encoding))
+	if (!FGLTFTextureUtility::ReadPixels(RenderTarget2D, *Pixels))
 	{
 		Builder.LogWarning(FString::Printf(TEXT("Failed to read pixels for 2D render target %s"), *JsonTexture->Name));
 		return;
@@ -160,101 +92,3 @@ void FGLTFDelayedTextureRenderTarget2DTask::Process()
 	JsonTexture->Source = Builder.AddUniqueImage(Pixels, Size, bIgnoreAlpha, Type, JsonTexture->Name);
 	JsonTexture->Sampler = Builder.AddUniqueSampler(RenderTarget2D);
 }
-
-FString FGLTFDelayedTextureRenderTargetCubeTask::GetName()
-{
-	return RenderTargetCube->GetName();
-}
-
-void FGLTFDelayedTextureRenderTargetCubeTask::Process()
-{
-	FGLTFTextureUtility::FullyLoad(RenderTargetCube);
-	JsonTexture->Name = RenderTargetCube->GetName() + TEXT("_") + FGLTFJsonUtilities::GetValue(FGLTFCoreUtilities::ConvertCubeFace(CubeFace));
-
-	const bool bFromSRGB = RenderTargetCube->SRGB;
-	if (bFromSRGB != bToSRGB)
-	{
-		JsonTexture->Name += bToSRGB ? TEXT("_sRGB") : TEXT("_Linear");
-	}
-
-	// TODO: add optimized "happy path" if cube face doesn't need rotation
-
-	const UTexture2D* FaceTexture = FGLTFTextureUtility::CreateTextureFromCubeFace(RenderTargetCube, CubeFace);
-	if (FaceTexture == nullptr)
-	{
-		Builder.LogWarning(FString::Printf(TEXT("Failed to extract cube face %d for cubemap render target %s"), CubeFace, *RenderTargetCube->GetName()));
-		return;
-	}
-
-	const bool bIsHDR = FGLTFTextureUtility::IsHDR(RenderTargetCube);
-	const FIntPoint Size = { RenderTargetCube->SizeX, RenderTargetCube->SizeX };
-	UTextureRenderTarget2D* RenderTarget = FGLTFTextureUtility::CreateRenderTarget(Size, bIsHDR);
-
-	const float FaceRotation = FGLTFTextureUtility::GetCubeFaceRotation(CubeFace);
-	FGLTFTextureUtility::RotateTexture(RenderTarget, FaceTexture, FVector2D::ZeroVector, Size, FaceRotation);
-
-	if (bIsHDR)
-	{
-		JsonTexture->Encoding = Builder.GetTextureHDREncoding();
-	}
-
-	TGLTFSharedArray<FColor> Pixels;
-	if (!FGLTFTextureUtility::ReadPixels(RenderTarget, *Pixels, JsonTexture->Encoding))
-	{
-		Builder.LogWarning(FString::Printf(TEXT("Failed to read pixels (cube face %d) for cubemap render target %s"), CubeFace, *RenderTargetCube->GetName()));
-		return;
-	}
-
-	FGLTFTextureUtility::TransformColorSpace(*Pixels, bFromSRGB, bToSRGB);
-
-	const bool bIgnoreAlpha = FGLTFTextureUtility::IsAlphaless(RenderTargetCube->GetFormat());
-	const EGLTFTextureType Type = bIsHDR ? EGLTFTextureType::HDR : EGLTFTextureType::None;
-
-	JsonTexture->Source = Builder.AddUniqueImage(Pixels, Size, bIgnoreAlpha, Type, JsonTexture->Name);
-	JsonTexture->Sampler = Builder.AddUniqueSampler(RenderTargetCube);
-}
-
-#if WITH_EDITOR
-
-FString FGLTFDelayedTextureLightMapTask::GetName()
-{
-	return LightMap->GetName();
-}
-
-void FGLTFDelayedTextureLightMapTask::Process()
-{
-	FGLTFTextureUtility::FullyLoad(LightMap);
-	LightMap->GetName(JsonTexture->Name);
-
-	// NOTE: export of lightmaps via source data is used to work around issues with
-	// quality-loss due to incorrect gamma transformation when rendering to a canvas.
-
-	FTextureSource& Source = const_cast<FTextureSource&>(LightMap->Source);
-	if (!Source.IsValid())
-	{
-		Builder.LogWarning(FString::Printf(TEXT("Failed to export lightmap texture %s because of missing source data"), *JsonTexture->Name));
-		return;
-	}
-
-	const ETextureSourceFormat SourceFormat = Source.GetFormat();
-	if (SourceFormat != TSF_BGRA8)
-	{
-		Builder.LogWarning(FString::Printf(TEXT("Failed to export lightmap texture %s because of unsupported source format %s"), *JsonTexture->Name, *FGLTFNameUtility::GetName(SourceFormat)));
-		return;
-	}
-
-	const FIntPoint Size = { Source.GetSizeX(), Source.GetSizeY() };
-	const int64 ByteLength = Source.CalcMipSize(0);
-
-	const bool bIgnoreAlpha = false;
-	const EGLTFTextureType Type = EGLTFTextureType::Lightmaps;
-
-	const void* RawData = Source.LockMip(0);
-	TGLTFSharedArray<FColor> Pixels = MakeShared<TArray<FColor>>(static_cast<const FColor*>(RawData), ByteLength / sizeof(FColor));
-	Source.UnlockMip(0);
-
-	JsonTexture->Source = Builder.AddUniqueImage(Pixels, Size, bIgnoreAlpha, Type, JsonTexture->Name);
-	JsonTexture->Sampler = Builder.AddUniqueSampler(LightMap);
-}
-
-#endif
diff --git a/Engine/Plugins/Enterprise/GLTFExporter/Source/GLTFExporter/Private/Tasks/GLTFDelayedTextureTasks.h b/Engine/Plugins/Enterprise/GLTFExporter/Source/GLTFExporter/Private/Tasks/GLTFDelayedTextureTasks.h
--- a/Engine/Plugins/Enterprise/GLTFExporter/Source/GLTFExporter/Private/Tasks/GLTFDelayedTextureTasks.h
+++ b/Engine/Plugins/Enterprise/GLTFExporter/Source/GLTFExporter/Private/Tasks/GLTFDelayedTextureTasks.h
@@ -30,33 +30,6 @@ private:
 	FGLTFJsonTexture* JsonTexture;
 };
 
-class FGLTFDelayedTextureCubeTask : public FGLTFDelayedTask
-{
-public:
-
-	FGLTFDelayedTextureCubeTask(FGLTFConvertBuilder& Builder, const UTextureCube* TextureCube, ECubeFace CubeFace, bool bToSRGB, FGLTFJsonTexture* JsonTexture)
-		: FGLTFDelayedTask(EGLTFTaskPriority::Texture)
-		, Builder(Builder)
-		, TextureCube(TextureCube)
-		, CubeFace(CubeFace)
-		, bToSRGB(bToSRGB)
-		, JsonTexture(JsonTexture)
-	{
-	}
-
-	virtual FString GetName() override;
-
-	virtual void Process() override;
-
-private:
-
-	FGLTFConvertBuilder& Builder;
-	const UTextureCube* TextureCube;
-	ECubeFace CubeFace;
-	bool bToSRGB;
-	FGLTFJsonTexture* JsonTexture;
-};
-
 class FGLTFDelayedTextureRenderTarget2DTask : public FGLTFDelayedTask
 {
 public:
@@ -81,57 +54,3 @@ private:
 	bool bToSRGB;
 	FGLTFJsonTexture* JsonTexture;
 };
-
-class FGLTFDelayedTextureRenderTargetCubeTask : public FGLTFDelayedTask
-{
-public:
-
-	FGLTFDelayedTextureRenderTargetCubeTask(FGLTFConvertBuilder& Builder, const UTextureRenderTargetCube* RenderTargetCube, ECubeFace CubeFace, bool bToSRGB, FGLTFJsonTexture* JsonTexture)
-		: FGLTFDelayedTask(EGLTFTaskPriority::Texture)
-		, Builder(Builder)
-		, RenderTargetCube(RenderTargetCube)
-		, CubeFace(CubeFace)
-		, bToSRGB(bToSRGB)
-		, JsonTexture(JsonTexture)
-	{
-	}
-
-	virtual FString GetName() override;
-
-	virtual void Process() override;
-
-private:
-
-	FGLTFConvertBuilder& Builder;
-	const UTextureRenderTargetCube* RenderTargetCube;
-	ECubeFace CubeFace;
-	bool bToSRGB;
-	FGLTFJsonTexture* JsonTexture;
-};
-
-#if WITH_EDITOR
-
-class FGLTFDelayedTextureLightMapTask : public FGLTFDelayedTask
-{
-public:
-
-	FGLTFDelayedTextureLightMapTask(FGLTFConvertBuilder& Builder, const ULightMapTexture2D* LightMap, FGLTFJsonTexture* JsonTexture)
-		: FGLTFDelayedTask(EGLTFTaskPriority::Texture)
-		, Builder(Builder)
-		, LightMap(LightMap)
-		, JsonTexture(JsonTexture)
-	{
-	}
-
-	virtual FString GetName() override;
-
-	virtual void Process() override;
-
-private:
-
-	FGLTFConvertBuilder& Builder;
-	const ULightMapTexture2D* LightMap;
-	FGLTFJsonTexture* JsonTexture;
-};
-
-#endif
diff --git a/Engine/Plugins/Enterprise/GLTFExporter/Source/GLTFExporter/Private/Utilities/GLTFCoreUtilities.cpp b/Engine/Plugins/Enterprise/GLTFExporter/Source/GLTFExporter/Private/Utilities/GLTFCoreUtilities.cpp
--- a/Engine/Plugins/Enterprise/GLTFExporter/Source/GLTFExporter/Private/Utilities/GLTFCoreUtilities.cpp
+++ b/Engine/Plugins/Enterprise/GLTFExporter/Source/GLTFExporter/Private/Utilities/GLTFCoreUtilities.cpp
@@ -245,24 +245,10 @@ EGLTFJsonAlphaMode FGLTFCoreUtilities::ConvertAlphaMode(EBlendMode Mode)
 		case BLEND_Opaque:         return EGLTFJsonAlphaMode::Opaque;
 		case BLEND_Masked:         return EGLTFJsonAlphaMode::Mask;
 		case BLEND_Translucent:    return EGLTFJsonAlphaMode::Blend;
-		case BLEND_Additive:       return EGLTFJsonAlphaMode::Blend;
-		case BLEND_Modulate:       return EGLTFJsonAlphaMode::Blend;
-		case BLEND_AlphaComposite: return EGLTFJsonAlphaMode::Blend;
 		default:                   return EGLTFJsonAlphaMode::None;
 	}
 }
 
-EGLTFJsonBlendMode FGLTFCoreUtilities::ConvertBlendMode(EBlendMode Mode)
-{
-	switch (Mode)
-	{
-		case BLEND_Additive:       return EGLTFJsonBlendMode::Additive;
-		case BLEND_Modulate:       return EGLTFJsonBlendMode::Modulate;
-		case BLEND_AlphaComposite: return EGLTFJsonBlendMode::AlphaComposite;
-		default:                   return EGLTFJsonBlendMode::None;
-	}
-}
-
 EGLTFJsonTextureWrap FGLTFCoreUtilities::ConvertWrap(TextureAddress Address)
 {
 	switch (Address)
@@ -305,17 +291,3 @@ EGLTFJsonTextureFilter FGLTFCoreUtilities::ConvertMagFilter(TextureFilter Filter
 {
 	return ConvertMagFilter(Filter == TF_Default ? FGLTFTextureUtility::GetDefaultFilter(LODGroup) : Filter);
 }
-
-EGLTFJsonCubeFace FGLTFCoreUtilities::ConvertCubeFace(ECubeFace CubeFace)
-{
-	switch (CubeFace)
-	{
-		case CubeFace_PosX:	return EGLTFJsonCubeFace::NegX;
-		case CubeFace_NegX:	return EGLTFJsonCubeFace::PosX;
-		case CubeFace_PosY:	return EGLTFJsonCubeFace::PosZ;
-		case CubeFace_NegY:	return EGLTFJsonCubeFace::NegZ;
-		case CubeFace_PosZ:	return EGLTFJsonCubeFace::PosY;
-		case CubeFace_NegZ:	return EGLTFJsonCubeFace::NegY;
-		default:            return EGLTFJsonCubeFace::None; // TODO: add error handling in callers
-	}
-}
diff --git a/Engine/Plugins/Enterprise/GLTFExporter/Source/GLTFExporter/Private/Utilities/GLTFJsonUtilities.cpp b/Engine/Plugins/Enterprise/GLTFExporter/Source/GLTFExporter/Private/Utilities/GLTFJsonUtilities.cpp
--- a/Engine/Plugins/Enterprise/GLTFExporter/Source/GLTFExporter/Private/Utilities/GLTFJsonUtilities.cpp
+++ b/Engine/Plugins/Enterprise/GLTFExporter/Source/GLTFExporter/Private/Utilities/GLTFJsonUtilities.cpp
@@ -12,13 +12,6 @@ const TCHAR* FGLTFJsonUtilities::GetValue(EGLTFJsonExtension Enum)
 		case EGLTFJsonExtension::KHR_MaterialsVariants:   return TEXT("KHR_materials_variants");
 		case EGLTFJsonExtension::KHR_MeshQuantization:    return TEXT("KHR_mesh_quantization");
 		case EGLTFJsonExtension::KHR_TextureTransform:    return TEXT("KHR_texture_transform");
-		case EGLTFJsonExtension::EPIC_AnimationPlayback:  return TEXT("EPIC_animation_playback");
-		case EGLTFJsonExtension::EPIC_BlendModes:         return TEXT("EPIC_blend_modes");
-		case EGLTFJsonExtension::EPIC_HDRIBackdrops:      return TEXT("EPIC_hdri_backdrops");
-		case EGLTFJsonExtension::EPIC_LevelVariantSets:   return TEXT("EPIC_level_variant_sets");
-		case EGLTFJsonExtension::EPIC_LightmapTextures:   return TEXT("EPIC_lightmap_textures");
-		case EGLTFJsonExtension::EPIC_SkySpheres:         return TEXT("EPIC_sky_spheres");
-		case EGLTFJsonExtension::EPIC_TextureHDREncoding: return TEXT("EPIC_texture_hdr_encoding");
 		default:
 			checkNoEntry();
 			return TEXT("");
@@ -38,19 +31,6 @@ const TCHAR* FGLTFJsonUtilities::GetValue(EGLTFJsonAlphaMode Enum)
 	}
 }
 
-const TCHAR* FGLTFJsonUtilities::GetValue(EGLTFJsonBlendMode Enum)
-{
-	switch (Enum)
-	{
-		case EGLTFJsonBlendMode::Additive:       return TEXT("ADDITIVE");
-		case EGLTFJsonBlendMode::Modulate:       return TEXT("MODULATE");
-		case EGLTFJsonBlendMode::AlphaComposite: return TEXT("ALPHACOMPOSITE");
-		default:
-			checkNoEntry();
-			return TEXT("");
-	}
-}
-
 const TCHAR* FGLTFJsonUtilities::GetValue(EGLTFJsonMimeType Enum)
 {
 	switch (Enum)
@@ -80,34 +60,6 @@ const TCHAR* FGLTFJsonUtilities::GetValue(EGLTFJsonAccessorType Enum)
 	}
 }
 
-const TCHAR* FGLTFJsonUtilities::GetValue(EGLTFJsonHDREncoding Enum)
-{
-	switch (Enum)
-	{
-		case EGLTFJsonHDREncoding::RGBM: return TEXT("RGBM");
-		case EGLTFJsonHDREncoding::RGBE: return TEXT("RGBE");
-		default:
-			checkNoEntry();
-			return TEXT("");
-	}
-}
-
-const TCHAR* FGLTFJsonUtilities::GetValue(EGLTFJsonCubeFace Enum)
-{
-	switch (Enum)
-	{
-		case EGLTFJsonCubeFace::PosX: return TEXT("PosX");
-		case EGLTFJsonCubeFace::NegX: return TEXT("NegX");
-		case EGLTFJsonCubeFace::PosY: return TEXT("PosY");
-		case EGLTFJsonCubeFace::NegY: return TEXT("NegY");
-		case EGLTFJsonCubeFace::PosZ: return TEXT("PosZ");
-		case EGLTFJsonCubeFace::NegZ: return TEXT("NegZ");
-		default:
-			checkNoEntry();
-			return TEXT("");
-	}
-}
-
 const TCHAR* FGLTFJsonUtilities::GetValue(EGLTFJsonCameraType Enum)
 {
 	switch (Enum)
diff --git a/Engine/Plugins/Enterprise/GLTFExporter/Source/GLTFExporter/Public/Builders/GLTFBuilder.h b/Engine/Plugins/Enterprise/GLTFExporter/Source/GLTFExporter/Public/Builders/GLTFBuilder.h
--- a/Engine/Plugins/Enterprise/GLTFExporter/Source/GLTFExporter/Public/Builders/GLTFBuilder.h
+++ b/Engine/Plugins/Enterprise/GLTFExporter/Source/GLTFExporter/Public/Builders/GLTFBuilder.h
@@ -31,8 +31,6 @@ public:
 	TextureFilter GetBakeFilterForMaterialProperty(const UMaterialInterface* Material, EGLTFMaterialPropertyGroup PropertyGroup) const;
 	TextureAddress GetBakeTilingForMaterialProperty(const UMaterialInterface* Material, EGLTFMaterialPropertyGroup PropertyGroup) const;
 
-	EGLTFJsonHDREncoding GetTextureHDREncoding() const;
-
 	bool ShouldExportLight(EComponentMobility::Type LightMobility) const;
 
 	int32 SanitizeLOD(const UStaticMesh* StaticMesh, const UStaticMeshComponent* StaticMeshComponent, int32 LODIndex) const;
diff --git a/Engine/Plugins/Enterprise/GLTFExporter/Source/GLTFExporter/Public/Builders/GLTFConvertBuilder.h b/Engine/Plugins/Enterprise/GLTFExporter/Source/GLTFExporter/Public/Builders/GLTFConvertBuilder.h
--- a/Engine/Plugins/Enterprise/GLTFExporter/Source/GLTFExporter/Public/Builders/GLTFConvertBuilder.h
+++ b/Engine/Plugins/Enterprise/GLTFExporter/Source/GLTFExporter/Public/Builders/GLTFConvertBuilder.h
@@ -16,13 +16,10 @@
 #include "Converters/GLTFSceneConverters.h"
 #include "Converters/GLTFCameraConverters.h"
 #include "Converters/GLTFLightConverters.h"
-#include "Converters/GLTFBackdropConverters.h"
-#include "Converters/GLTFEpicLevelVariantSetsConverters.h"
-#include "Converters/GLTFLightMapConverters.h"
 #include "Converters/GLTFKhrVariantConverters.h"
-#include "Converters/GLTFSkySphereConverters.h"
 
 class UMeshComponent;
+class UPropertyValue;
 
 class GLTFEXPORTER_API FGLTFConvertBuilder : public FGLTFBufferBuilder
 {
@@ -63,15 +60,10 @@ public:
 	FGLTFJsonSampler* AddUniqueSampler(const UTexture* Texture);
 	FGLTFJsonTexture* AddUniqueTexture(const UTexture* Texture);
 	FGLTFJsonTexture* AddUniqueTexture(const UTexture2D* Texture);
-	FGLTFJsonTexture* AddUniqueTexture(const UTextureCube* Texture, ECubeFace CubeFace);
 	FGLTFJsonTexture* AddUniqueTexture(const UTextureRenderTarget2D* Texture);
-	FGLTFJsonTexture* AddUniqueTexture(const UTextureRenderTargetCube* Texture, ECubeFace CubeFace);
-	FGLTFJsonTexture* AddUniqueTexture(const ULightMapTexture2D* Texture);
 	FGLTFJsonTexture* AddUniqueTexture(const UTexture* Texture, bool bToSRGB);
 	FGLTFJsonTexture* AddUniqueTexture(const UTexture2D* Texture, bool bToSRGB);
-	FGLTFJsonTexture* AddUniqueTexture(const UTextureCube* Texture, ECubeFace CubeFace, bool bToSRGB);
 	FGLTFJsonTexture* AddUniqueTexture(const UTextureRenderTarget2D* Texture, bool bToSRGB);
-	FGLTFJsonTexture* AddUniqueTexture(const UTextureRenderTargetCube* Texture, ECubeFace CubeFace, bool bToSRGB);
 	FGLTFJsonImage* AddUniqueImage(TGLTFSharedArray<FColor>& Pixels, FIntPoint Size, bool bIgnoreAlpha, EGLTFTextureType Type, const FString& Name);
 
 	FGLTFJsonSkin* AddUniqueSkin(FGLTFJsonNode* RootNode, const USkeletalMesh* SkeletalMesh);
@@ -91,11 +83,7 @@ public:
 
 	FGLTFJsonCamera* AddUniqueCamera(const UCameraComponent* CameraComponent);
 	FGLTFJsonLight* AddUniqueLight(const ULightComponent* LightComponent);
-	FGLTFJsonBackdrop* AddUniqueBackdrop(const AActor* BackdropActor);
-	FGLTFJsonLightMap* AddUniqueLightMap(const UStaticMeshComponent* StaticMeshComponent);
-	FGLTFJsonSkySphere* AddUniqueSkySphere(const AActor* SkySphereActor);
 
-	FGLTFJsonEpicLevelVariantSets* AddUniqueEpicLevelVariantSets(const ULevelVariantSets* LevelVariantSets);
 	FGLTFJsonKhrMaterialVariant* AddUniqueKhrMaterialVariant(const UVariant* Variant);
 
 	void RegisterObjectVariant(const UObject* Object, const UPropertyValue* Property);
@@ -119,10 +107,7 @@ public:
 
 	TUniquePtr<IGLTFSamplerConverter> SamplerConverter = MakeUnique<FGLTFSamplerConverter>(*this);
 	TUniquePtr<IGLTFTexture2DConverter> Texture2DConverter = MakeUnique<FGLTFTexture2DConverter>(*this);
-	TUniquePtr<IGLTFTextureCubeConverter> TextureCubeConverter = MakeUnique<FGLTFTextureCubeConverter>(*this);
 	TUniquePtr<IGLTFTextureRenderTarget2DConverter> TextureRenderTarget2DConverter = MakeUnique<FGLTFTextureRenderTarget2DConverter>(*this);
-	TUniquePtr<IGLTFTextureRenderTargetCubeConverter> TextureRenderTargetCubeConverter = MakeUnique<FGLTFTextureRenderTargetCubeConverter>(*this);
-	TUniquePtr<IGLTFTextureLightMapConverter> TextureLightMapConverter = MakeUnique<FGLTFTextureLightMapConverter>(*this);
 	TUniquePtr<IGLTFImageConverter> ImageConverter = MakeUnique<FGLTFImageConverter>(*this);
 
 	TUniquePtr<IGLTFSkinConverter> SkinConverter = MakeUnique<FGLTFSkinConverter>(*this);
@@ -141,11 +126,6 @@ public:
 
 	TUniquePtr<IGLTFCameraConverter> CameraConverter = MakeUnique<FGLTFCameraConverter>(*this);
 	TUniquePtr<IGLTFLightConverter> LightConverter = MakeUnique<FGLTFLightConverter>(*this);
-	TUniquePtr<IGLTFBackdropConverter> BackdropConverter = MakeUnique<FGLTFBackdropConverter>(*this);
-	TUniquePtr<IGLTFLightMapConverter> LightMapConverter = MakeUnique<FGLTFLightMapConverter>(*this);
-	TUniquePtr<IGLTFSkySphereConverter> SkySphereConverter = MakeUnique<FGLTFSkySphereConverter>(*this);
-
-	TUniquePtr<IGLTFEpicLevelVariantSetsConverter> EpicLevelVariantSetsConverter = MakeUnique<FGLTFEpicLevelVariantSetsConverter>(*this);
 	TUniquePtr<IGLTFKhrMaterialVariantConverter> KhrMaterialVariantConverter = MakeUnique<FGLTFKhrMaterialVariantConverter>(*this);
 
 private:
diff --git a/Engine/Plugins/Enterprise/GLTFExporter/Source/GLTFExporter/Public/Builders/GLTFJsonBuilder.h b/Engine/Plugins/Enterprise/GLTFExporter/Source/GLTFExporter/Public/Builders/GLTFJsonBuilder.h
--- a/Engine/Plugins/Enterprise/GLTFExporter/Source/GLTFExporter/Public/Builders/GLTFJsonBuilder.h
+++ b/Engine/Plugins/Enterprise/GLTFExporter/Source/GLTFExporter/Public/Builders/GLTFJsonBuilder.h
@@ -28,11 +28,7 @@ public:
 	FGLTFJsonScene* AddScene();
 	FGLTFJsonSkin* AddSkin();
 	FGLTFJsonTexture* AddTexture();
-	FGLTFJsonBackdrop* AddBackdrop();
 	FGLTFJsonLight* AddLight();
-	FGLTFJsonLightMap* AddLightMap();
-	FGLTFJsonSkySphere* AddSkySphere();
-	FGLTFJsonEpicLevelVariantSets* AddEpicLevelVariantSets();
 	FGLTFJsonKhrMaterialVariant* AddKhrMaterialVariant();
 
 	const FGLTFJsonRoot& GetRoot() const;
@@ -41,13 +37,9 @@ protected:
 
 	void WriteJsonArchive(FArchive& Archive);
 
-	TSet<EGLTFJsonExtension> GetCustomExtensionsUsed() const;
-
 private:
 
 	FString GetGeneratorString() const;
 
-	static bool IsCustomExtension(EGLTFJsonExtension Extension);
-
 	FGLTFJsonRoot JsonRoot;
 };
diff --git a/Engine/Plugins/Enterprise/GLTFExporter/Source/GLTFExporter/Public/Converters/GLTFBackdropConverters.h b/Engine/Plugins/Enterprise/GLTFExporter/Source/GLTFExporter/Public/Converters/GLTFBackdropConverters.h
deleted file mode 100644
diff --git a/Engine/Plugins/Enterprise/GLTFExporter/Source/GLTFExporter/Public/Converters/GLTFEpicLevelVariantSetsConverters.h b/Engine/Plugins/Enterprise/GLTFExporter/Source/GLTFExporter/Public/Converters/GLTFEpicLevelVariantSetsConverters.h
deleted file mode 100644
diff --git a/Engine/Plugins/Enterprise/GLTFExporter/Source/GLTFExporter/Public/Converters/GLTFLightMapConverters.h b/Engine/Plugins/Enterprise/GLTFExporter/Source/GLTFExporter/Public/Converters/GLTFLightMapConverters.h
deleted file mode 100644
diff --git a/Engine/Plugins/Enterprise/GLTFExporter/Source/GLTFExporter/Public/Converters/GLTFSkySphereConverters.h b/Engine/Plugins/Enterprise/GLTFExporter/Source/GLTFExporter/Public/Converters/GLTFSkySphereConverters.h
deleted file mode 100644
diff --git a/Engine/Plugins/Enterprise/GLTFExporter/Source/GLTFExporter/Public/Converters/GLTFTextureConverters.h b/Engine/Plugins/Enterprise/GLTFExporter/Source/GLTFExporter/Public/Converters/GLTFTextureConverters.h
--- a/Engine/Plugins/Enterprise/GLTFExporter/Source/GLTFExporter/Public/Converters/GLTFTextureConverters.h
+++ b/Engine/Plugins/Enterprise/GLTFExporter/Source/GLTFExporter/Public/Converters/GLTFTextureConverters.h
@@ -5,22 +5,14 @@
 #include "Json/GLTFJsonCore.h"
 #include "Converters/GLTFConverter.h"
 #include "Converters/GLTFBuilderContext.h"
-#include "Engine/LightMapTexture2D.h"
-#include "Engine/Texture2D.h"
-#include "Engine/TextureCube.h"
-#include "Engine/TextureRenderTarget2D.h"
-#include "Engine/TextureRenderTargetCube.h"
-#include "RHIDefinitions.h"
 
-// TODO: generalize parameter bToSRGB to SamplerType (including normalmap unpacking)
+class UTexture2D;
+class UTextureRenderTarget2D;
 
-// TODO: remove lightmap-specific converter (since it won't work in runtime)
+// TODO: generalize parameter bToSRGB to SamplerType (including normalmap unpacking)
 
 typedef TGLTFConverter<FGLTFJsonTexture*, const UTexture2D*, bool> IGLTFTexture2DConverter;
-typedef TGLTFConverter<FGLTFJsonTexture*, const UTextureCube*, ECubeFace, bool> IGLTFTextureCubeConverter;
 typedef TGLTFConverter<FGLTFJsonTexture*, const UTextureRenderTarget2D*, bool> IGLTFTextureRenderTarget2DConverter;
-typedef TGLTFConverter<FGLTFJsonTexture*, const UTextureRenderTargetCube*, ECubeFace, bool> IGLTFTextureRenderTargetCubeConverter;
-typedef TGLTFConverter<FGLTFJsonTexture*, const ULightMapTexture2D*> IGLTFTextureLightMapConverter;
 
 class GLTFEXPORTER_API FGLTFTexture2DConverter : public FGLTFBuilderContext, public IGLTFTexture2DConverter
 {
@@ -30,24 +22,9 @@ public:
 
 protected:
 
-	virtual void Sanitize(const UTexture2D*& Texture2D, bool& bToSRGB) override;
-
 	virtual FGLTFJsonTexture* Convert(const UTexture2D* Texture2D, bool bToSRGB) override;
 };
 
-class GLTFEXPORTER_API FGLTFTextureCubeConverter : public FGLTFBuilderContext, public IGLTFTextureCubeConverter
-{
-public:
-
-	using FGLTFBuilderContext::FGLTFBuilderContext;
-
-protected:
-
-	virtual void Sanitize(const UTextureCube*& TextureCube, ECubeFace& CubeFace, bool& bToSRGB) override;
-
-	virtual FGLTFJsonTexture* Convert(const UTextureCube* TextureCube, ECubeFace CubeFace, bool bToSRGB) override;
-};
-
 class GLTFEXPORTER_API FGLTFTextureRenderTarget2DConverter : public FGLTFBuilderContext, public IGLTFTextureRenderTarget2DConverter
 {
 public:
@@ -56,31 +33,5 @@ public:
 
 protected:
 
-	virtual void Sanitize(const UTextureRenderTarget2D*& RenderTarget2D, bool& bToSRGB) override;
-
 	virtual FGLTFJsonTexture* Convert(const UTextureRenderTarget2D* RenderTarget2D, bool bToSRGB) override;
 };
-
-class GLTFEXPORTER_API FGLTFTextureRenderTargetCubeConverter : public FGLTFBuilderContext, public IGLTFTextureRenderTargetCubeConverter
-{
-public:
-
-	using FGLTFBuilderContext::FGLTFBuilderContext;
-
-protected:
-
-	virtual void Sanitize(const UTextureRenderTargetCube*& RenderTargetCube, ECubeFace& CubeFace, bool& bToSRGB) override;
-
-	virtual FGLTFJsonTexture* Convert(const UTextureRenderTargetCube* RenderTargetCube, ECubeFace CubeFace, bool bToSRGB) override;
-};
-
-class GLTFEXPORTER_API FGLTFTextureLightMapConverter : public FGLTFBuilderContext, public IGLTFTextureLightMapConverter
-{
-public:
-
-	using FGLTFBuilderContext::FGLTFBuilderContext;
-
-protected:
-
-	virtual FGLTFJsonTexture* Convert(const ULightMapTexture2D* LightMap) override;
-};
diff --git a/Engine/Plugins/Enterprise/GLTFExporter/Source/GLTFExporter/Public/Json/GLTFJsonAnimation.h b/Engine/Plugins/Enterprise/GLTFExporter/Source/GLTFExporter/Public/Json/GLTFJsonAnimation.h
--- a/Engine/Plugins/Enterprise/GLTFExporter/Source/GLTFExporter/Public/Json/GLTFJsonAnimation.h
+++ b/Engine/Plugins/Enterprise/GLTFExporter/Source/GLTFExporter/Public/Json/GLTFJsonAnimation.h
@@ -3,7 +3,6 @@
 #pragma once
 
 #include "Json/GLTFJsonCore.h"
-#include "Json/GLTFJsonAnimationPlayback.h"
 
 struct GLTFEXPORTER_API FGLTFJsonAnimationChannelTarget : IGLTFJsonObject
 {
@@ -61,8 +60,6 @@ struct GLTFEXPORTER_API FGLTFJsonAnimation : IGLTFJsonIndexedObject
 	TArray<FGLTFJsonAnimationChannel> Channels;
 	TGLTFJsonIndexedObjectArray<FGLTFJsonAnimationSampler> Samplers;
 
-	FGLTFJsonAnimationPlayback Playback;
-
 	virtual void WriteObject(IGLTFJsonWriter& Writer) const override;
 
 protected:
diff --git a/Engine/Plugins/Enterprise/GLTFExporter/Source/GLTFExporter/Public/Json/GLTFJsonAnimationPlayback.h b/Engine/Plugins/Enterprise/GLTFExporter/Source/GLTFExporter/Public/Json/GLTFJsonAnimationPlayback.h
deleted file mode 100644
diff --git a/Engine/Plugins/Enterprise/GLTFExporter/Source/GLTFExporter/Public/Json/GLTFJsonBackdrop.h b/Engine/Plugins/Enterprise/GLTFExporter/Source/GLTFExporter/Public/Json/GLTFJsonBackdrop.h
deleted file mode 100644
diff --git a/Engine/Plugins/Enterprise/GLTFExporter/Source/GLTFExporter/Public/Json/GLTFJsonEnums.h b/Engine/Plugins/Enterprise/GLTFExporter/Source/GLTFExporter/Public/Json/GLTFJsonEnums.h
--- a/Engine/Plugins/Enterprise/GLTFExporter/Source/GLTFExporter/Public/Json/GLTFJsonEnums.h
+++ b/Engine/Plugins/Enterprise/GLTFExporter/Source/GLTFExporter/Public/Json/GLTFJsonEnums.h
@@ -9,14 +9,7 @@ enum class EGLTFJsonExtension
 	KHR_MaterialsUnlit,
 	KHR_MaterialsVariants,
 	KHR_MeshQuantization,
-	KHR_TextureTransform,
-	EPIC_AnimationPlayback,
-	EPIC_BlendModes,
-	EPIC_HDRIBackdrops,
-	EPIC_LevelVariantSets,
-	EPIC_LightmapTextures,
-	EPIC_SkySpheres,
-	EPIC_TextureHDREncoding
+	KHR_TextureTransform
 };
 
 enum class EGLTFJsonShadingModel
@@ -29,24 +22,6 @@ enum class EGLTFJsonShadingModel
 	NumShadingModels
 };
 
-enum class EGLTFJsonHDREncoding
-{
-	None = -1,
-	RGBM,
-	RGBE
-};
-
-enum class EGLTFJsonCubeFace
-{
-	None = -1,
-	PosX,
-	NegX,
-	PosY,
-	NegY,
-	PosZ,
-	NegZ
-};
-
 enum class EGLTFJsonAccessorType
 {
 	None = -1,
@@ -97,14 +72,6 @@ enum class EGLTFJsonAlphaMode
 	Blend
 };
 
-enum class EGLTFJsonBlendMode
-{
-	None = -1,
-	Additive,
-	Modulate,
-	AlphaComposite
-};
-
 enum class EGLTFJsonMimeType
 {
 	None = -1,
diff --git a/Engine/Plugins/Enterprise/GLTFExporter/Source/GLTFExporter/Public/Json/GLTFJsonEpicLevelVariantSets.h b/Engine/Plugins/Enterprise/GLTFExporter/Source/GLTFExporter/Public/Json/GLTFJsonEpicLevelVariantSets.h
deleted file mode 100644
diff --git a/Engine/Plugins/Enterprise/GLTFExporter/Source/GLTFExporter/Public/Json/GLTFJsonLightMap.h b/Engine/Plugins/Enterprise/GLTFExporter/Source/GLTFExporter/Public/Json/GLTFJsonLightMap.h
deleted file mode 100644
diff --git a/Engine/Plugins/Enterprise/GLTFExporter/Source/GLTFExporter/Public/Json/GLTFJsonMaterial.h b/Engine/Plugins/Enterprise/GLTFExporter/Source/GLTFExporter/Public/Json/GLTFJsonMaterial.h
--- a/Engine/Plugins/Enterprise/GLTFExporter/Source/GLTFExporter/Public/Json/GLTFJsonMaterial.h
+++ b/Engine/Plugins/Enterprise/GLTFExporter/Source/GLTFExporter/Public/Json/GLTFJsonMaterial.h
@@ -102,8 +102,6 @@ struct GLTFEXPORTER_API FGLTFJsonMaterial : IGLTFJsonIndexedObject
 
 	bool DoubleSided;
 
-	EGLTFJsonBlendMode BlendMode;
-
 	FGLTFJsonClearCoatExtension ClearCoat;
 
 	virtual void WriteObject(IGLTFJsonWriter& Writer) const override;
@@ -119,7 +117,6 @@ protected:
 		, AlphaMode(EGLTFJsonAlphaMode::Opaque)
 		, AlphaCutoff(0.5f)
 		, DoubleSided(false)
-		, BlendMode(EGLTFJsonBlendMode::None)
 	{
 	}
 };
diff --git a/Engine/Plugins/Enterprise/GLTFExporter/Source/GLTFExporter/Public/Json/GLTFJsonNode.h b/Engine/Plugins/Enterprise/GLTFExporter/Source/GLTFExporter/Public/Json/GLTFJsonNode.h
--- a/Engine/Plugins/Enterprise/GLTFExporter/Source/GLTFExporter/Public/Json/GLTFJsonNode.h
+++ b/Engine/Plugins/Enterprise/GLTFExporter/Source/GLTFExporter/Public/Json/GLTFJsonNode.h
@@ -15,10 +15,7 @@ struct GLTFEXPORTER_API FGLTFJsonNode : IGLTFJsonIndexedObject
 	FGLTFJsonCamera*    Camera;
 	FGLTFJsonSkin*      Skin;
 	FGLTFJsonMesh*      Mesh;
-	FGLTFJsonBackdrop*  Backdrop;
 	FGLTFJsonLight*     Light;
-	FGLTFJsonLightMap*  LightMap;
-	FGLTFJsonSkySphere* SkySphere;
 
 	TArray<FGLTFJsonNode*> Children;
 
@@ -36,10 +33,7 @@ protected:
 		, Camera(nullptr)
 		, Skin(nullptr)
 		, Mesh(nullptr)
-		, Backdrop(nullptr)
 		, Light(nullptr)
-		, LightMap(nullptr)
-		, SkySphere(nullptr)
 	{
 	}
 };
diff --git a/Engine/Plugins/Enterprise/GLTFExporter/Source/GLTFExporter/Public/Json/GLTFJsonRoot.h b/Engine/Plugins/Enterprise/GLTFExporter/Source/GLTFExporter/Public/Json/GLTFJsonRoot.h
--- a/Engine/Plugins/Enterprise/GLTFExporter/Source/GLTFExporter/Public/Json/GLTFJsonRoot.h
+++ b/Engine/Plugins/Enterprise/GLTFExporter/Source/GLTFExporter/Public/Json/GLTFJsonRoot.h
@@ -17,11 +17,7 @@
 #include "Json/GLTFJsonScene.h"
 #include "Json/GLTFJsonSkin.h"
 #include "Json/GLTFJsonTexture.h"
-#include "Json/GLTFJsonBackdrop.h"
 #include "Json/GLTFJsonLight.h"
-#include "Json/GLTFJsonLightMap.h"
-#include "Json/GLTFJsonSkySphere.h"
-#include "Json/GLTFJsonEpicLevelVariantSets.h"
 #include "Json/GLTFJsonKhrMaterialVariant.h"
 
 struct GLTFEXPORTER_API FGLTFJsonRoot : IGLTFJsonObject
@@ -45,11 +41,7 @@ struct GLTFEXPORTER_API FGLTFJsonRoot : IGLTFJsonObject
 	TGLTFJsonIndexedObjectArray<FGLTFJsonScene>      Scenes;
 	TGLTFJsonIndexedObjectArray<FGLTFJsonSkin>       Skins;
 	TGLTFJsonIndexedObjectArray<FGLTFJsonTexture>    Textures;
-	TGLTFJsonIndexedObjectArray<FGLTFJsonBackdrop>   Backdrops;
 	TGLTFJsonIndexedObjectArray<FGLTFJsonLight>      Lights;
-	TGLTFJsonIndexedObjectArray<FGLTFJsonLightMap>   LightMaps;
-	TGLTFJsonIndexedObjectArray<FGLTFJsonSkySphere>  SkySpheres;
-	TGLTFJsonIndexedObjectArray<FGLTFJsonEpicLevelVariantSets> EpicLevelVariantSets;
 	TGLTFJsonIndexedObjectArray<FGLTFJsonKhrMaterialVariant>   KhrMaterialVariants;
 
 	FGLTFJsonRoot()
diff --git a/Engine/Plugins/Enterprise/GLTFExporter/Source/GLTFExporter/Public/Json/GLTFJsonScene.h b/Engine/Plugins/Enterprise/GLTFExporter/Source/GLTFExporter/Public/Json/GLTFJsonScene.h
--- a/Engine/Plugins/Enterprise/GLTFExporter/Source/GLTFExporter/Public/Json/GLTFJsonScene.h
+++ b/Engine/Plugins/Enterprise/GLTFExporter/Source/GLTFExporter/Public/Json/GLTFJsonScene.h
@@ -9,7 +9,6 @@ struct GLTFEXPORTER_API FGLTFJsonScene : IGLTFJsonIndexedObject
 	FString Name;
 
 	TArray<FGLTFJsonNode*> Nodes;
-	TArray<FGLTFJsonEpicLevelVariantSets*>  EpicLevelVariantSets;
 
 	virtual void WriteObject(IGLTFJsonWriter& Writer) const override;
 
diff --git a/Engine/Plugins/Enterprise/GLTFExporter/Source/GLTFExporter/Public/Json/GLTFJsonSkySphere.h b/Engine/Plugins/Enterprise/GLTFExporter/Source/GLTFExporter/Public/Json/GLTFJsonSkySphere.h
deleted file mode 100644
diff --git a/Engine/Plugins/Enterprise/GLTFExporter/Source/GLTFExporter/Public/Json/GLTFJsonTexture.h b/Engine/Plugins/Enterprise/GLTFExporter/Source/GLTFExporter/Public/Json/GLTFJsonTexture.h
--- a/Engine/Plugins/Enterprise/GLTFExporter/Source/GLTFExporter/Public/Json/GLTFJsonTexture.h
+++ b/Engine/Plugins/Enterprise/GLTFExporter/Source/GLTFExporter/Public/Json/GLTFJsonTexture.h
@@ -12,8 +12,6 @@ struct GLTFEXPORTER_API FGLTFJsonTexture : IGLTFJsonIndexedObject
 
 	FGLTFJsonImage* Source;
 
-	EGLTFJsonHDREncoding Encoding;
-
 	virtual void WriteObject(IGLTFJsonWriter& Writer) const override;
 
 protected:
@@ -24,7 +22,6 @@ protected:
 		: IGLTFJsonIndexedObject(Index)
 		, Sampler(nullptr)
 		, Source(nullptr)
-		, Encoding(EGLTFJsonHDREncoding::None)
 	{
 	}
 };
diff --git a/Engine/Plugins/Enterprise/GLTFExporter/Source/GLTFExporter/Public/Options/GLTFExportOptions.h b/Engine/Plugins/Enterprise/GLTFExporter/Source/GLTFExporter/Public/Options/GLTFExportOptions.h
--- a/Engine/Plugins/Enterprise/GLTFExporter/Source/GLTFExporter/Public/Options/GLTFExportOptions.h
+++ b/Engine/Plugins/Enterprise/GLTFExporter/Source/GLTFExporter/Public/Options/GLTFExportOptions.h
@@ -52,17 +52,6 @@ enum class EGLTFSceneMobility : uint8
 };
 ENUM_CLASS_FLAGS(EGLTFSceneMobility);
 
-UENUM(BlueprintType)
-enum class EGLTFVariantSetsMode : uint8
-{
-	/** Never export variants sets. */
-	None,
-	/** Uses the official extension KHR_materials_variants. Supports material variants only. */
-	Khronos,
-	/** Uses the extension EPIC_level_variant_sets, which is supported by Unreal's glTF viewer. */
-	Epic
-};
-
 UENUM(BlueprintType)
 enum class EGLTFMaterialVariantMode : uint8
 {
@@ -111,10 +100,6 @@ class GLTFEXPORTER_API UGLTFExportOptions : public UObject
 	UPROPERTY(EditAnywhere, BlueprintReadWrite, Config, Category = Material)
 	bool bExportClearCoatMaterials;
 
-	/** If enabled, materials with blend modes additive, modulate, and alpha composite will be properly exported. Uses extension EPIC_blend_modes, which is supported by Unreal's glTF viewer. */
-	UPROPERTY(EditAnywhere, BlueprintReadWrite, Config, Category = Material)
-	bool bExportExtraBlendModes;
-
 	/** Bake mode determining if and how a material input is baked out to a texture. Baking is only used for non-trivial material inputs (i.e. not simple texture or constant expressions). */
 	UPROPERTY(EditAnywhere, BlueprintReadWrite, Config, Category = Material)
 	EGLTFMaterialBakeMode BakeMaterialInputs;
@@ -159,10 +144,6 @@ class GLTFEXPORTER_API UGLTFExportOptions : public UObject
 	UPROPERTY(EditAnywhere, BlueprintReadWrite, Config, Category = Animation, Meta = (EditCondition = "bExportVertexSkinWeights"))
 	bool bExportAnimationSequences;
 
-	/** If enabled, export play rate, start time, looping, and auto play for an animation or level sequence. Uses extension EPIC_animation_playback, which is supported by Unreal's glTF viewer. */
-	UPROPERTY(EditAnywhere, BlueprintReadWrite, Config, Category = Animation)
-	bool bExportPlaybackSettings;
-
 	/** Desired image format used for exported textures. */
 	UPROPERTY(EditAnywhere, BlueprintReadWrite, Config, Category = Texture)
 	EGLTFTextureImageFormat TextureImageFormat;
@@ -179,14 +160,6 @@ class GLTFEXPORTER_API UGLTFExportOptions : public UObject
 	UPROPERTY(EditAnywhere, BlueprintReadWrite, Config, Category = Texture, Meta = (EditCondition = "TextureImageFormat != EGLTFTextureImageFormat::None"))
 	bool bExportTextureTransforms;
 
-	/** If enabled, export lightmaps (created by Lightmass) when exporting a level. Uses extension EPIC_lightmap_textures, which is supported by Unreal's glTF viewer. */
-	UPROPERTY(EditAnywhere, BlueprintReadWrite, Config, Category = Texture, Meta = (EditCondition = "TextureImageFormat != EGLTFTextureImageFormat::None"))
-	bool bExportLightmaps;
-
-	/** Encoding used to store textures that have pixel colors with more than 8-bit per channel. Uses extension EPIC_texture_hdr_encoding, which is supported by Unreal's glTF viewer. */
-	UPROPERTY(EditAnywhere, BlueprintReadWrite, Config, Category = Texture, Meta = (DisplayName = "Texture HDR Encoding", EditCondition = "TextureImageFormat != EGLTFTextureImageFormat::None"))
-	EGLTFTextureHDREncoding TextureHDREncoding;
-
 	/** If enabled, exported normalmaps will be adjusted from Unreal to glTF convention (i.e. the green channel is flipped). */
 	UPROPERTY(EditAnywhere, BlueprintReadWrite, Config, Category = Texture, Meta = (EditCondition = "TextureImageFormat != EGLTFTextureImageFormat::None"))
 	bool bAdjustNormalmaps;
@@ -203,30 +176,10 @@ class GLTFEXPORTER_API UGLTFExportOptions : public UObject
 	UPROPERTY(EditAnywhere, BlueprintReadWrite, Config, Category = Scene)
 	bool bExportCameras;
 
-	/** If enabled, export HDRIBackdrop blueprints. Uses extension EPIC_hdri_backdrops, which is supported by Unreal's glTF viewer. */
-	UPROPERTY(EditAnywhere, BlueprintReadWrite, Config, Category = Scene, Meta = (DisplayName = "Export HDRI Backdrops"))
-	bool bExportHDRIBackdrops;
-
-	/** If enabled, export SkySphere blueprints. Uses extension EPIC_sky_spheres, which is supported by Unreal's glTF viewer. */
-	UPROPERTY(EditAnywhere, BlueprintReadWrite, Config, Category = Scene)
-	bool bExportSkySpheres;
-
-	/** Mode determining if and how to export LevelVariantSetsActors. */
-	UPROPERTY(EditAnywhere, BlueprintReadWrite, Config, Category = Scene)
-	EGLTFVariantSetsMode VariantSetsMode;
-
 	/** Mode determining if and how to export material variants that change the materials property on a static or skeletal mesh component. */
 	UPROPERTY(EditAnywhere, BlueprintReadWrite, Config, Category = VariantSets, Meta = (EditCondition = "VariantSetsMode != EGLTFVariantSetsMode::None"))
 	EGLTFMaterialVariantMode ExportMaterialVariants;
 
-	/** If enabled, export variants that change the mesh property on a static or skeletal mesh component. */
-	UPROPERTY(EditAnywhere, BlueprintReadWrite, Config, Category = VariantSets, Meta = (EditCondition = "VariantSetsMode == EGLTFVariantSetsMode::Epic"))
-	bool bExportMeshVariants;
-
-	/** If enabled, export variants that change the visible property on a scene component. */
-	UPROPERTY(EditAnywhere, BlueprintReadWrite, Config, Category = VariantSets, Meta = (EditCondition = "VariantSetsMode == EGLTFVariantSetsMode::Epic"))
-	bool bExportVisibilityVariants;
-
 	UFUNCTION(BlueprintCallable, Category = General)
 	void ResetToDefault();
 };
diff --git a/Engine/Plugins/Enterprise/GLTFExporter/Source/GLTFExporter/Public/Utilities/GLTFCoreUtilities.h b/Engine/Plugins/Enterprise/GLTFExporter/Source/GLTFExporter/Public/Utilities/GLTFCoreUtilities.h
--- a/Engine/Plugins/Enterprise/GLTFExporter/Source/GLTFExporter/Public/Utilities/GLTFCoreUtilities.h
+++ b/Engine/Plugins/Enterprise/GLTFExporter/Source/GLTFExporter/Public/Utilities/GLTFCoreUtilities.h
@@ -62,7 +62,6 @@ struct GLTFEXPORTER_API FGLTFCoreUtilities
 	static EGLTFJsonShadingModel ConvertShadingModel(EMaterialShadingModel ShadingModel);
 
 	static EGLTFJsonAlphaMode ConvertAlphaMode(EBlendMode Mode);
-	static EGLTFJsonBlendMode ConvertBlendMode(EBlendMode Mode);
 
 	static EGLTFJsonTextureWrap ConvertWrap(TextureAddress Address);
 
@@ -72,8 +71,6 @@ struct GLTFEXPORTER_API FGLTFCoreUtilities
 	static EGLTFJsonTextureFilter ConvertMinFilter(TextureFilter Filter, TextureGroup LODGroup);
 	static EGLTFJsonTextureFilter ConvertMagFilter(TextureFilter Filter, TextureGroup LODGroup);
 
-	static EGLTFJsonCubeFace ConvertCubeFace(ECubeFace CubeFace);
-
 	template <typename ComponentType>
 	static EGLTFJsonComponentType GetComponentType()
 	{
diff --git a/Engine/Plugins/Enterprise/GLTFExporter/Source/GLTFExporter/Public/Utilities/GLTFJsonUtilities.h b/Engine/Plugins/Enterprise/GLTFExporter/Source/GLTFExporter/Public/Utilities/GLTFJsonUtilities.h
--- a/Engine/Plugins/Enterprise/GLTFExporter/Source/GLTFExporter/Public/Utilities/GLTFJsonUtilities.h
+++ b/Engine/Plugins/Enterprise/GLTFExporter/Source/GLTFExporter/Public/Utilities/GLTFJsonUtilities.h
@@ -15,11 +15,8 @@ struct GLTFEXPORTER_API FGLTFJsonUtilities
 
 	static const TCHAR* GetValue(EGLTFJsonExtension Enum);
 	static const TCHAR* GetValue(EGLTFJsonAlphaMode Enum);
-	static const TCHAR* GetValue(EGLTFJsonBlendMode Enum);
 	static const TCHAR* GetValue(EGLTFJsonMimeType Enum);
 	static const TCHAR* GetValue(EGLTFJsonAccessorType Enum);
-	static const TCHAR* GetValue(EGLTFJsonHDREncoding Enum);
-	static const TCHAR* GetValue(EGLTFJsonCubeFace Enum);
 	static const TCHAR* GetValue(EGLTFJsonCameraType Enum);
 	static const TCHAR* GetValue(EGLTFJsonLightType Enum);
 	static const TCHAR* GetValue(EGLTFJsonInterpolation Enum);
