Fix issue with incorrect alpha factor when exporting materials that have constant basecolor and opacity after material baking

diff --git a/Engine/Plugins/Enterprise/GLTFExporter/Source/GLTFExporter/Private/Tasks/GLTFDelayedMaterialTasks.cpp b/Engine/Plugins/Enterprise/GLTFExporter/Source/GLTFExporter/Private/Tasks/GLTFDelayedMaterialTasks.cpp
--- a/Engine/Plugins/Enterprise/GLTFExporter/Source/GLTFExporter/Private/Tasks/GLTFDelayedMaterialTasks.cpp
+++ b/Engine/Plugins/Enterprise/GLTFExporter/Source/GLTFExporter/Private/Tasks/GLTFDelayedMaterialTasks.cpp
@@ -521,7 +521,7 @@ bool FGLTFDelayedMaterialTask::TryGetBaseColorAndOpacity(FGLTFJsonPBRMetallicRou
 	if (BaseColorBakeOutput.bIsConstant && OpacityBakeOutput.bIsConstant)
 	{
 		FLinearColor BaseColorFactor(BaseColorBakeOutput.ConstantValue * BaseColorScale);
-		BaseColorFactor.A = OpacityBakeOutput.ConstantValue.A;
+		BaseColorFactor.A = OpacityBakeOutput.ConstantValue.R;
 
 		OutPBRParams.BaseColorFactor = FGLTFCoreUtilities::ConvertColor(BaseColorFactor, Builder.ExportOptions->bStrictCompliance);
 		return true;
