Fix nonunity compiler error on mac in FGLTFSamplerConverter

diff --git a/Engine/Plugins/Enterprise/GLTFExporter/Source/GLTFExporter/Private/Converters/GLTFSamplerConverters.cpp b/Engine/Plugins/Enterprise/GLTFExporter/Source/GLTFExporter/Private/Converters/GLTFSamplerConverters.cpp
--- a/Engine/Plugins/Enterprise/GLTFExporter/Source/GLTFExporter/Private/Converters/GLTFSamplerConverters.cpp
+++ b/Engine/Plugins/Enterprise/GLTFExporter/Source/GLTFExporter/Private/Converters/GLTFSamplerConverters.cpp
@@ -12,17 +12,8 @@ FGLTFJsonSampler* FGLTFSamplerConverter::Convert(const UTexture* Texture)
 	FGLTFJsonSampler* JsonSampler = Builder.AddSampler();
 	Texture->GetName(JsonSampler->Name);
 
-	if (Texture->IsA<ULightMapTexture2D>())
-	{
-		// Override default filter settings for LightMap textures (which otherwise is "nearest")
-		JsonSampler->MinFilter = EGLTFJsonTextureFilter::LinearMipmapLinear;
-		JsonSampler->MagFilter = EGLTFJsonTextureFilter::Linear;
-	}
-	else
-	{
-		JsonSampler->MinFilter = FGLTFCoreUtilities::ConvertMinFilter(Texture->Filter, Texture->LODGroup);
-		JsonSampler->MagFilter = FGLTFCoreUtilities::ConvertMagFilter(Texture->Filter, Texture->LODGroup);
-	}
+	JsonSampler->MinFilter = FGLTFCoreUtilities::ConvertMinFilter(Texture->Filter, Texture->LODGroup);
+	JsonSampler->MagFilter = FGLTFCoreUtilities::ConvertMagFilter(Texture->Filter, Texture->LODGroup);
 
 	if (FGLTFTextureUtility::IsCubemap(Texture))
 	{
