Remove "Khr" prefix for all material variant types since it will no longer conflict with the removed custom "Epic" variants

diff --git a/Engine/Plugins/Enterprise/GLTFExporter/Source/GLTFExporter/Private/Builders/GLTFConvertBuilder.cpp b/Engine/Plugins/Enterprise/GLTFExporter/Source/GLTFExporter/Private/Builders/GLTFConvertBuilder.cpp
--- a/Engine/Plugins/Enterprise/GLTFExporter/Source/GLTFExporter/Private/Builders/GLTFConvertBuilder.cpp
+++ b/Engine/Plugins/Enterprise/GLTFExporter/Source/GLTFExporter/Private/Builders/GLTFConvertBuilder.cpp
@@ -462,14 +462,14 @@ FGLTFJsonLight* FGLTFConvertBuilder::AddUniqueLight(const ULightComponent* Light
 	return LightConverter->GetOrAdd(LightComponent);
 }
 
-FGLTFJsonKhrMaterialVariant* FGLTFConvertBuilder::AddUniqueKhrMaterialVariant(const UVariant* Variant)
+FGLTFJsonMaterialVariant* FGLTFConvertBuilder::AddUniqueMaterialVariant(const UVariant* Variant)
 {
 	if (Variant == nullptr)
 	{
 		return nullptr;
 	}
 
-	return KhrMaterialVariantConverter->GetOrAdd(Variant);
+	return MaterialVariantConverter->GetOrAdd(Variant);
 }
 
 void FGLTFConvertBuilder::RegisterObjectVariant(const UObject* Object, const UPropertyValue* Property)
diff --git a/Engine/Plugins/Enterprise/GLTFExporter/Source/GLTFExporter/Private/Builders/GLTFJsonBuilder.cpp b/Engine/Plugins/Enterprise/GLTFExporter/Source/GLTFExporter/Private/Builders/GLTFJsonBuilder.cpp
--- a/Engine/Plugins/Enterprise/GLTFExporter/Source/GLTFExporter/Private/Builders/GLTFJsonBuilder.cpp
+++ b/Engine/Plugins/Enterprise/GLTFExporter/Source/GLTFExporter/Private/Builders/GLTFJsonBuilder.cpp
@@ -96,9 +96,9 @@ FGLTFJsonLight* FGLTFJsonBuilder::AddLight()
 	return JsonRoot.Lights.Add();
 }
 
-FGLTFJsonKhrMaterialVariant* FGLTFJsonBuilder::AddKhrMaterialVariant()
+FGLTFJsonMaterialVariant* FGLTFJsonBuilder::AddMaterialVariant()
 {
-	return JsonRoot.KhrMaterialVariants.Add();
+	return JsonRoot.MaterialVariants.Add();
 }
 
 const FGLTFJsonRoot& FGLTFJsonBuilder::GetRoot() const
diff --git a/Engine/Plugins/Enterprise/GLTFExporter/Source/GLTFExporter/Private/Converters/GLTFKhrVariantConverters.cpp b/Engine/Plugins/Enterprise/GLTFExporter/Source/GLTFExporter/Private/Converters/GLTFKhrVariantConverters.cpp
deleted file mode 100644
diff --git a/Engine/Plugins/Enterprise/GLTFExporter/Source/GLTFExporter/Private/Converters/GLTFMaterialVariantConverters.cpp b/Engine/Plugins/Enterprise/GLTFExporter/Source/GLTFExporter/Private/Converters/GLTFMaterialVariantConverters.cpp
new file mode 100644
--- /dev/null
+++ b/Engine/Plugins/Enterprise/GLTFExporter/Source/GLTFExporter/Private/Converters/GLTFMaterialVariantConverters.cpp
@@ -0,0 +1,149 @@
+﻿// Copyright Epic Games, Inc. All Rights Reserved.
+
+#include "Converters/GLTFMaterialVariantConverters.h"
+#include "Converters/GLTFVariantUtility.h"
+#include "Builders/GLTFContainerBuilder.h"
+#include "Components/MeshComponent.h"
+#include "VariantObjectBinding.h"
+#include "PropertyValueMaterial.h"
+#include "PropertyValue.h"
+#include "Variant.h"
+
+FGLTFJsonMaterialVariant* FGLTFMaterialVariantConverter::Convert(const UVariant* Variant)
+{
+	if (Variant == nullptr || Builder.ExportOptions->ExportMaterialVariants == EGLTFMaterialVariantMode::None)
+	{
+		return nullptr;
+	}
+
+	typedef TTuple<FGLTFJsonPrimitive*, FGLTFJsonMaterial*> FGLTFPrimitiveMaterial;
+	TArray<TTuple<FGLTFJsonPrimitive*, FGLTFJsonMaterial*>> PrimitiveMaterials;
+
+	for (const UVariantObjectBinding* Binding: Variant->GetBindings())
+	{
+		for (const UPropertyValue* Property: Binding->GetCapturedProperties())
+		{
+			if (!const_cast<UPropertyValue*>(Property)->Resolve() || !Property->HasRecordedData())
+			{
+				continue;
+			}
+
+			if (const UPropertyValueMaterial* MaterialProperty = Cast<UPropertyValueMaterial>(Property))
+			{
+				FGLTFJsonPrimitive* Primitive = nullptr;
+				FGLTFJsonMaterial* JsonMaterial;
+
+				if (TryParseMaterialProperty(Primitive, JsonMaterial, MaterialProperty))
+				{
+					PrimitiveMaterials.Add(MakeTuple(Primitive, JsonMaterial));
+				}
+			}
+		}
+	}
+
+	if (PrimitiveMaterials.Num() < 1)
+	{
+		// TODO: add warning and / or allow unused material variants to be added?
+
+		return nullptr;
+	}
+
+	FGLTFJsonMaterialVariant* MaterialVariant = Builder.AddMaterialVariant();
+	// TODO: add warning if the variant name is not unique, i.e it's already used?
+	// While material variants are technically allowed to use the same name, it may
+	// cause confusion when trying to select the correct variant in a viewer.
+	MaterialVariant->Name = Variant->GetDisplayText().ToString();
+
+	for (const FGLTFPrimitiveMaterial& PrimitiveMaterial: PrimitiveMaterials)
+	{
+		FGLTFJsonPrimitive* Primitive = PrimitiveMaterial.Key;
+		FGLTFJsonMaterial* Material = PrimitiveMaterial.Value;
+
+		FGLTFJsonMaterialVariantMapping* ExistingMapping = Primitive->MaterialVariantMappings.FindByPredicate(
+			[Material](const FGLTFJsonMaterialVariantMapping& Mapping)
+			{
+				return Mapping.Material == Material;
+			});
+
+		if (ExistingMapping != nullptr)
+		{
+			ExistingMapping->Variants.AddUnique(MaterialVariant);
+		}
+		else
+		{
+			FGLTFJsonMaterialVariantMapping Mapping;
+			Mapping.Material = Material;
+			Mapping.Variants.Add(MaterialVariant);
+
+			Primitive->MaterialVariantMappings.Add(Mapping);
+		}
+	}
+
+	return MaterialVariant;
+}
+
+bool FGLTFMaterialVariantConverter::TryParseMaterialProperty(FGLTFJsonPrimitive*& OutPrimitive, FGLTFJsonMaterial*& OutMaterial, const UPropertyValueMaterial* Property) const
+{
+	const UMeshComponent* Target = static_cast<UMeshComponent*>(Property->GetPropertyParentContainerAddress());
+	if (Target == nullptr)
+	{
+		Builder.LogWarning(FString::Printf(
+			TEXT("Variant property %s must belong to a mesh component, the property will be skipped"),
+			*FGLTFVariantUtility::GetLogContext(Property)));
+		return false;
+	}
+
+	const AActor* Owner = Target->GetOwner();
+	if (Owner == nullptr)
+	{
+		Builder.LogWarning(FString::Printf(
+			TEXT("Variant property %s must belong to an actor, the property will be skipped"),
+			*FGLTFVariantUtility::GetLogContext(Property)));
+		return false;
+	}
+
+	if (!Builder.IsSelectedActor(Owner))
+	{
+		Builder.LogWarning(FString::Printf(
+			TEXT("Variant property %s doesn't belong to an actor selected for export, the property will be skipped"),
+			*FGLTFVariantUtility::GetLogContext(Property)));
+		return false;
+	}
+
+	const TArray<FCapturedPropSegment>& CapturedPropSegments = Property->GetCapturedPropSegments();
+	const int32 NumPropSegments = CapturedPropSegments.Num();
+
+	if (NumPropSegments < 1)
+	{
+		Builder.LogWarning(FString::Printf(
+			TEXT("Failed to parse material index for variant property %s, the property will be skipped"),
+			*FGLTFVariantUtility::GetLogContext(Property)));
+		return false;
+	}
+
+	// NOTE: UPropertyValueMaterial::GetMaterial does *not* ensure that the recorded data has been loaded,
+	// so we need to call UProperty::GetRecordedData first to make that happen.
+	const_cast<UPropertyValueMaterial*>(Property)->GetRecordedData();
+
+	const UMaterialInterface* Material = const_cast<UPropertyValueMaterial*>(Property)->GetMaterial();
+	if (Material == nullptr)
+	{
+		// TODO: find way to determine whether the material is null because "None" was selected, or because it failed to resolve
+
+		Builder.LogWarning(FString::Printf(
+			TEXT("No material assigned, the property will be skipped. Context: %s"),
+			*FGLTFVariantUtility::GetLogContext(Property)));
+		return false;
+	}
+
+	Builder.RegisterObjectVariant(Target, Property);
+	Builder.RegisterObjectVariant(Owner, Property); // TODO: we don't need to register this on the actor
+
+	const int32 MaterialIndex = CapturedPropSegments[NumPropSegments - 1].PropertyIndex;
+	FGLTFJsonMesh* JsonMesh = Builder.AddUniqueMesh(Target);
+
+	OutPrimitive = &JsonMesh->Primitives[MaterialIndex];
+	OutMaterial = FGLTFVariantUtility::AddUniqueMaterial(Builder, Material, Target, MaterialIndex);
+
+	return true;
+}
diff --git a/Engine/Plugins/Enterprise/GLTFExporter/Source/GLTFExporter/Private/Converters/GLTFSceneConverters.cpp b/Engine/Plugins/Enterprise/GLTFExporter/Source/GLTFExporter/Private/Converters/GLTFSceneConverters.cpp
--- a/Engine/Plugins/Enterprise/GLTFExporter/Source/GLTFExporter/Private/Converters/GLTFSceneConverters.cpp
+++ b/Engine/Plugins/Enterprise/GLTFExporter/Source/GLTFExporter/Private/Converters/GLTFSceneConverters.cpp
@@ -42,7 +42,7 @@ FGLTFJsonScene* FGLTFSceneConverter::Convert(const UWorld* World)
 							{
 								for (const UVariant* Variant: VariantSet->GetVariants())
 								{
-									Builder.AddUniqueKhrMaterialVariant(Variant);
+									Builder.AddUniqueMaterialVariant(Variant);
 								}
 							}
 						}
diff --git a/Engine/Plugins/Enterprise/GLTFExporter/Source/GLTFExporter/Private/Exporters/GLTFLevelVariantSetsExporter.cpp b/Engine/Plugins/Enterprise/GLTFExporter/Source/GLTFExporter/Private/Exporters/GLTFLevelVariantSetsExporter.cpp
--- a/Engine/Plugins/Enterprise/GLTFExporter/Source/GLTFExporter/Private/Exporters/GLTFLevelVariantSetsExporter.cpp
+++ b/Engine/Plugins/Enterprise/GLTFExporter/Source/GLTFExporter/Private/Exporters/GLTFLevelVariantSetsExporter.cpp
@@ -52,7 +52,7 @@ bool UGLTFLevelVariantSetsExporter::AddObject(FGLTFContainerBuilder& Builder, co
 		return false;
 	}
 
-	if (Builder.GetRoot().KhrMaterialVariants.Num() < 1)
+	if (Builder.GetRoot().MaterialVariants.Num() < 1)
 	{
 		Builder.LogError(
 			FString::Printf(TEXT("Failed to export any supported variant from level variant sets %s"),
diff --git a/Engine/Plugins/Enterprise/GLTFExporter/Source/GLTFExporter/Private/Json/GLTFJsonKhrMaterialVariant.cpp b/Engine/Plugins/Enterprise/GLTFExporter/Source/GLTFExporter/Private/Json/GLTFJsonKhrMaterialVariant.cpp
deleted file mode 100644
diff --git a/Engine/Plugins/Enterprise/GLTFExporter/Source/GLTFExporter/Private/Json/GLTFJsonMaterialVariant.cpp b/Engine/Plugins/Enterprise/GLTFExporter/Source/GLTFExporter/Private/Json/GLTFJsonMaterialVariant.cpp
new file mode 100644
--- /dev/null
+++ b/Engine/Plugins/Enterprise/GLTFExporter/Source/GLTFExporter/Private/Json/GLTFJsonMaterialVariant.cpp
@@ -0,0 +1,15 @@
+// Copyright Epic Games, Inc. All Rights Reserved.
+
+#include "Json/GLTFJsonMaterialVariant.h"
+#include "Json/GLTFJsonMaterial.h"
+
+void FGLTFJsonMaterialVariantMapping::WriteObject(IGLTFJsonWriter& Writer) const
+{
+	Writer.Write(TEXT("material"), Material);
+	Writer.Write(TEXT("variants"), Variants);
+}
+
+void FGLTFJsonMaterialVariant::WriteObject(IGLTFJsonWriter& Writer) const
+{
+	Writer.Write(TEXT("name"), Name);
+}
diff --git a/Engine/Plugins/Enterprise/GLTFExporter/Source/GLTFExporter/Private/Json/GLTFJsonMesh.cpp b/Engine/Plugins/Enterprise/GLTFExporter/Source/GLTFExporter/Private/Json/GLTFJsonMesh.cpp
--- a/Engine/Plugins/Enterprise/GLTFExporter/Source/GLTFExporter/Private/Json/GLTFJsonMesh.cpp
+++ b/Engine/Plugins/Enterprise/GLTFExporter/Source/GLTFExporter/Private/Json/GLTFJsonMesh.cpp
@@ -49,12 +49,12 @@ void FGLTFJsonPrimitive::WriteObject(IGLTFJsonWriter& Writer) const
 		Writer.Write(TEXT("mode"), Mode);
 	}
 
-	if (KhrMaterialVariantMappings.Num() > 0)
+	if (MaterialVariantMappings.Num() > 0)
 	{
 		Writer.StartExtensions();
 
 		Writer.StartExtension(EGLTFJsonExtension::KHR_MaterialsVariants);
-		Writer.Write(TEXT("mappings"), KhrMaterialVariantMappings);
+		Writer.Write(TEXT("mappings"), MaterialVariantMappings);
 		Writer.EndExtension();
 
 		Writer.EndExtensions();
diff --git a/Engine/Plugins/Enterprise/GLTFExporter/Source/GLTFExporter/Private/Json/GLTFJsonRoot.cpp b/Engine/Plugins/Enterprise/GLTFExporter/Source/GLTFExporter/Private/Json/GLTFJsonRoot.cpp
--- a/Engine/Plugins/Enterprise/GLTFExporter/Source/GLTFExporter/Private/Json/GLTFJsonRoot.cpp
+++ b/Engine/Plugins/Enterprise/GLTFExporter/Source/GLTFExporter/Private/Json/GLTFJsonRoot.cpp
@@ -25,14 +25,14 @@ void FGLTFJsonRoot::WriteObject(IGLTFJsonWriter& Writer) const
 	if (Skins.Num() > 0) Writer.Write(TEXT("skins"), Skins);
 	if (Textures.Num() > 0) Writer.Write(TEXT("textures"), Textures);
 
-	if (Lights.Num() > 0 || KhrMaterialVariants.Num() > 0)
+	if (Lights.Num() > 0 || MaterialVariants.Num() > 0)
 	{
 		Writer.StartExtensions();
 
-		if (KhrMaterialVariants.Num() > 0)
+		if (MaterialVariants.Num() > 0)
 		{
 			Writer.StartExtension(EGLTFJsonExtension::KHR_MaterialsVariants);
-			Writer.Write(TEXT("variants"), KhrMaterialVariants);
+			Writer.Write(TEXT("variants"), MaterialVariants);
 			Writer.EndExtension();
 		}
 
diff --git a/Engine/Plugins/Enterprise/GLTFExporter/Source/GLTFExporter/Public/Builders/GLTFConvertBuilder.h b/Engine/Plugins/Enterprise/GLTFExporter/Source/GLTFExporter/Public/Builders/GLTFConvertBuilder.h
--- a/Engine/Plugins/Enterprise/GLTFExporter/Source/GLTFExporter/Public/Builders/GLTFConvertBuilder.h
+++ b/Engine/Plugins/Enterprise/GLTFExporter/Source/GLTFExporter/Public/Builders/GLTFConvertBuilder.h
@@ -16,7 +16,7 @@
 #include "Converters/GLTFSceneConverters.h"
 #include "Converters/GLTFCameraConverters.h"
 #include "Converters/GLTFLightConverters.h"
-#include "Converters/GLTFKhrVariantConverters.h"
+#include "Converters/GLTFMaterialVariantConverters.h"
 
 class UMeshComponent;
 class UPropertyValue;
@@ -84,7 +84,7 @@ public:
 	FGLTFJsonCamera* AddUniqueCamera(const UCameraComponent* CameraComponent);
 	FGLTFJsonLight* AddUniqueLight(const ULightComponent* LightComponent);
 
-	FGLTFJsonKhrMaterialVariant* AddUniqueKhrMaterialVariant(const UVariant* Variant);
+	FGLTFJsonMaterialVariant* AddUniqueMaterialVariant(const UVariant* Variant);
 
 	void RegisterObjectVariant(const UObject* Object, const UPropertyValue* Property);
 	const TArray<const UPropertyValue*>* GetObjectVariants(const UObject* Object) const;
@@ -126,7 +126,7 @@ public:
 
 	TUniquePtr<IGLTFCameraConverter> CameraConverter = MakeUnique<FGLTFCameraConverter>(*this);
 	TUniquePtr<IGLTFLightConverter> LightConverter = MakeUnique<FGLTFLightConverter>(*this);
-	TUniquePtr<IGLTFKhrMaterialVariantConverter> KhrMaterialVariantConverter = MakeUnique<FGLTFKhrMaterialVariantConverter>(*this);
+	TUniquePtr<IGLTFMaterialVariantConverter> MaterialVariantConverter = MakeUnique<FGLTFMaterialVariantConverter>(*this);
 
 private:
 
diff --git a/Engine/Plugins/Enterprise/GLTFExporter/Source/GLTFExporter/Public/Builders/GLTFJsonBuilder.h b/Engine/Plugins/Enterprise/GLTFExporter/Source/GLTFExporter/Public/Builders/GLTFJsonBuilder.h
--- a/Engine/Plugins/Enterprise/GLTFExporter/Source/GLTFExporter/Public/Builders/GLTFJsonBuilder.h
+++ b/Engine/Plugins/Enterprise/GLTFExporter/Source/GLTFExporter/Public/Builders/GLTFJsonBuilder.h
@@ -29,7 +29,7 @@ public:
 	FGLTFJsonSkin* AddSkin();
 	FGLTFJsonTexture* AddTexture();
 	FGLTFJsonLight* AddLight();
-	FGLTFJsonKhrMaterialVariant* AddKhrMaterialVariant();
+	FGLTFJsonMaterialVariant* AddMaterialVariant();
 
 	const FGLTFJsonRoot& GetRoot() const;
 
diff --git a/Engine/Plugins/Enterprise/GLTFExporter/Source/GLTFExporter/Public/Converters/GLTFKhrVariantConverters.h b/Engine/Plugins/Enterprise/GLTFExporter/Source/GLTFExporter/Public/Converters/GLTFKhrVariantConverters.h
deleted file mode 100644
diff --git a/Engine/Plugins/Enterprise/GLTFExporter/Source/GLTFExporter/Public/Converters/GLTFMaterialVariantConverters.h b/Engine/Plugins/Enterprise/GLTFExporter/Source/GLTFExporter/Public/Converters/GLTFMaterialVariantConverters.h
new file mode 100644
--- /dev/null
+++ b/Engine/Plugins/Enterprise/GLTFExporter/Source/GLTFExporter/Public/Converters/GLTFMaterialVariantConverters.h
@@ -0,0 +1,28 @@
+﻿// Copyright Epic Games, Inc. All Rights Reserved.
+
+#pragma once
+
+#include "Json/GLTFJsonCore.h"
+#include "Json/GLTFJsonMesh.h"
+#include "Converters/GLTFConverter.h"
+#include "Converters/GLTFBuilderContext.h"
+
+class UVariant;
+class UPropertyValueMaterial;
+
+typedef TGLTFConverter<FGLTFJsonMaterialVariant*, const UVariant*> IGLTFMaterialVariantConverter;
+
+class GLTFEXPORTER_API FGLTFMaterialVariantConverter : public FGLTFBuilderContext, public IGLTFMaterialVariantConverter
+{
+public:
+
+	using FGLTFBuilderContext::FGLTFBuilderContext;
+
+protected:
+
+	virtual FGLTFJsonMaterialVariant* Convert(const UVariant* Variant) override;
+
+private:
+
+	bool TryParseMaterialProperty(FGLTFJsonPrimitive*& OutPrimitive, FGLTFJsonMaterial*& OutMaterial, const UPropertyValueMaterial* Property) const;
+};
diff --git a/Engine/Plugins/Enterprise/GLTFExporter/Source/GLTFExporter/Public/Json/GLTFJsonCore.h b/Engine/Plugins/Enterprise/GLTFExporter/Source/GLTFExporter/Public/Json/GLTFJsonCore.h
--- a/Engine/Plugins/Enterprise/GLTFExporter/Source/GLTFExporter/Public/Json/GLTFJsonCore.h
+++ b/Engine/Plugins/Enterprise/GLTFExporter/Source/GLTFExporter/Public/Json/GLTFJsonCore.h
@@ -35,4 +35,4 @@ struct FGLTFJsonLight;
 struct FGLTFJsonLightMap;
 struct FGLTFJsonSkySphere;
 struct FGLTFJsonEpicLevelVariantSets;
-struct FGLTFJsonKhrMaterialVariant;
+struct FGLTFJsonMaterialVariant;
diff --git a/Engine/Plugins/Enterprise/GLTFExporter/Source/GLTFExporter/Public/Json/GLTFJsonKhrMaterialVariant.h b/Engine/Plugins/Enterprise/GLTFExporter/Source/GLTFExporter/Public/Json/GLTFJsonKhrMaterialVariant.h
deleted file mode 100644
diff --git a/Engine/Plugins/Enterprise/GLTFExporter/Source/GLTFExporter/Public/Json/GLTFJsonMaterialVariant.h b/Engine/Plugins/Enterprise/GLTFExporter/Source/GLTFExporter/Public/Json/GLTFJsonMaterialVariant.h
new file mode 100644
--- /dev/null
+++ b/Engine/Plugins/Enterprise/GLTFExporter/Source/GLTFExporter/Public/Json/GLTFJsonMaterialVariant.h
@@ -0,0 +1,34 @@
+// Copyright Epic Games, Inc. All Rights Reserved.
+
+#pragma once
+
+#include "Json/GLTFJsonCore.h"
+
+struct GLTFEXPORTER_API FGLTFJsonMaterialVariantMapping : IGLTFJsonObject
+{
+	FGLTFJsonMaterial* Material;
+	TArray<FGLTFJsonMaterialVariant*> Variants;
+
+	FGLTFJsonMaterialVariantMapping()
+		: Material(nullptr)
+	{
+	}
+
+	virtual void WriteObject(IGLTFJsonWriter& Writer) const override;
+};
+
+struct GLTFEXPORTER_API FGLTFJsonMaterialVariant : IGLTFJsonIndexedObject
+{
+	FString Name;
+
+	virtual void WriteObject(IGLTFJsonWriter& Writer) const override;
+
+protected:
+
+	friend TGLTFJsonIndexedObjectArray<FGLTFJsonMaterialVariant, void>;
+
+	FGLTFJsonMaterialVariant(int32 Index)
+		: IGLTFJsonIndexedObject(Index)
+	{
+	}
+};
diff --git a/Engine/Plugins/Enterprise/GLTFExporter/Source/GLTFExporter/Public/Json/GLTFJsonMesh.h b/Engine/Plugins/Enterprise/GLTFExporter/Source/GLTFExporter/Public/Json/GLTFJsonMesh.h
--- a/Engine/Plugins/Enterprise/GLTFExporter/Source/GLTFExporter/Public/Json/GLTFJsonMesh.h
+++ b/Engine/Plugins/Enterprise/GLTFExporter/Source/GLTFExporter/Public/Json/GLTFJsonMesh.h
@@ -3,7 +3,7 @@
 #pragma once
 
 #include "Json/GLTFJsonCore.h"
-#include "Json/GLTFJsonKhrMaterialVariant.h"
+#include "Json/GLTFJsonMaterialVariant.h"
 
 struct GLTFEXPORTER_API FGLTFJsonAttributes : IGLTFJsonObject
 {
@@ -34,7 +34,7 @@ struct GLTFEXPORTER_API FGLTFJsonPrimitive : IGLTFJsonObject
 	FGLTFJsonMaterial* Material;
 	EGLTFJsonPrimitiveMode Mode;
 
-	TArray<FGLTFJsonKhrMaterialVariantMapping> KhrMaterialVariantMappings;
+	TArray<FGLTFJsonMaterialVariantMapping> MaterialVariantMappings;
 
 	FGLTFJsonPrimitive()
 		: Indices(nullptr)
diff --git a/Engine/Plugins/Enterprise/GLTFExporter/Source/GLTFExporter/Public/Json/GLTFJsonRoot.h b/Engine/Plugins/Enterprise/GLTFExporter/Source/GLTFExporter/Public/Json/GLTFJsonRoot.h
--- a/Engine/Plugins/Enterprise/GLTFExporter/Source/GLTFExporter/Public/Json/GLTFJsonRoot.h
+++ b/Engine/Plugins/Enterprise/GLTFExporter/Source/GLTFExporter/Public/Json/GLTFJsonRoot.h
@@ -18,7 +18,7 @@
 #include "Json/GLTFJsonSkin.h"
 #include "Json/GLTFJsonTexture.h"
 #include "Json/GLTFJsonLight.h"
-#include "Json/GLTFJsonKhrMaterialVariant.h"
+#include "Json/GLTFJsonMaterialVariant.h"
 
 struct GLTFEXPORTER_API FGLTFJsonRoot : IGLTFJsonObject
 {
@@ -42,7 +42,7 @@ struct GLTFEXPORTER_API FGLTFJsonRoot : IGLTFJsonObject
 	TGLTFJsonIndexedObjectArray<FGLTFJsonSkin>       Skins;
 	TGLTFJsonIndexedObjectArray<FGLTFJsonTexture>    Textures;
 	TGLTFJsonIndexedObjectArray<FGLTFJsonLight>      Lights;
-	TGLTFJsonIndexedObjectArray<FGLTFJsonKhrMaterialVariant>   KhrMaterialVariants;
+	TGLTFJsonIndexedObjectArray<FGLTFJsonMaterialVariant>   MaterialVariants;
 
 	FGLTFJsonRoot()
 		: DefaultScene(nullptr)
