Cleanup unused code in texture conversion, remaining from removed custom extensions

diff --git a/Engine/Plugins/Enterprise/GLTFExporter/Source/GLTFExporter/Private/Converters/GLTFTextureUtility.cpp b/Engine/Plugins/Enterprise/GLTFExporter/Source/GLTFExporter/Private/Converters/GLTFTextureUtility.cpp
--- a/Engine/Plugins/Enterprise/GLTFExporter/Source/GLTFExporter/Private/Converters/GLTFTextureUtility.cpp
+++ b/Engine/Plugins/Enterprise/GLTFExporter/Source/GLTFExporter/Private/Converters/GLTFTextureUtility.cpp
@@ -64,20 +64,6 @@ bool FGLTFTextureUtility::IsCubemap(const UTexture* Texture)
 	return Texture->IsA<UTextureCube>() || Texture->IsA<UTextureRenderTargetCube>();
 }
 
-float FGLTFTextureUtility::GetCubeFaceRotation(ECubeFace CubeFace)
-{
-	switch (CubeFace)
-	{
-		case CubeFace_PosX: return 90;
-		case CubeFace_NegX: return -90;
-		case CubeFace_PosY: return 180;
-		case CubeFace_NegY: return 0;
-		case CubeFace_PosZ: return 180;
-		case CubeFace_NegZ: return 0;
-		default:            return 0;
-	}
-}
-
 TextureFilter FGLTFTextureUtility::GetDefaultFilter(TextureGroup LODGroup)
 {
 	const UTextureLODSettings* TextureLODSettings = UDeviceProfileManager::Get().GetActiveProfile()->GetTextureLODSettings();
@@ -141,28 +127,6 @@ TTuple<TextureAddress, TextureAddress> FGLTFTextureUtility::GetAddressXY(const U
 	return MakeTuple(AddressX, AddressY);
 }
 
-UTexture2D* FGLTFTextureUtility::CreateTransientTexture(const void* RawData, int64 ByteLength, const FIntPoint& Size, EPixelFormat Format, bool bSRGB)
-{
-	check(CalculateImageBytes(Size.X, Size.Y, 0, Format) == ByteLength);
-
-	// TODO: do these temp textures need to be part of the root set to avoid garbage collection?
-	UTexture2D* Texture = UTexture2D::CreateTransient(Size.X, Size.Y, Format);
-
-	void* MipData = Texture->GetPlatformData()->Mips[0].BulkData.Lock(LOCK_READ_WRITE);
-	FMemory::Memcpy(MipData, RawData, ByteLength);
-	Texture->GetPlatformData()->Mips[0].BulkData.Unlock();
-
-	Texture->SRGB = bSRGB ? 1 : 0;
-	Texture->CompressionSettings = TC_VectorDisplacementmap; // best quality
-#if WITH_EDITOR
-	Texture->CompressionNone = true;
-	Texture->MipGenSettings = TMGS_NoMipmaps;
-#endif
-
-	Texture->UpdateResource();
-	return Texture;
-}
-
 UTextureRenderTarget2D* FGLTFTextureUtility::CreateRenderTarget(const FIntPoint& Size, bool bHDR)
 {
 	// TODO: instead of PF_FloatRGBA (i.e. RTF_RGBA16f) use PF_A32B32G32R32F (i.e. RTF_RGBA32f) to avoid accuracy loss
@@ -170,8 +134,8 @@ UTextureRenderTarget2D* FGLTFTextureUtility::CreateRenderTarget(const FIntPoint&
 
 	// NOTE: both bForceLinearGamma and TargetGamma=2.2 seem necessary for exported images to match their source data.
 	// It's not entirely clear why gamma must be 2.2 (instead of 0.0) and why bInForceLinearGamma must also be true.
-	const bool bForceLinearGamma = true;
-	const float TargetGamma = 2.2f;
+	constexpr bool bForceLinearGamma = true;
+	constexpr float TargetGamma = 2.2f;
 
 	UTextureRenderTarget2D* RenderTarget = NewObject<UTextureRenderTarget2D>();
 	RenderTarget->InitCustomFormat(Size.X, Size.Y, PixelFormat, bForceLinearGamma);
@@ -180,7 +144,7 @@ UTextureRenderTarget2D* FGLTFTextureUtility::CreateRenderTarget(const FIntPoint&
 	return RenderTarget;
 }
 
-bool FGLTFTextureUtility::DrawTexture(UTextureRenderTarget2D* OutTarget, const UTexture2D* InSource, const FVector2D& InPosition, const FVector2D& InSize, const FMatrix& InTransform)
+bool FGLTFTextureUtility::DrawTexture(UTextureRenderTarget2D* OutTarget, const UTexture2D* InSource)
 {
 	FRenderTarget* RenderTarget = OutTarget->GameThread_GetRenderTargetResource();
 	if (RenderTarget == nullptr)
@@ -202,12 +166,10 @@ bool FGLTFTextureUtility::DrawTexture(UTextureRenderTarget2D* OutTarget, const U
 	}
 
 	FCanvas Canvas(RenderTarget, nullptr, FGameTime::CreateDilated(0.0f, 0.0f, 0.0f, 0.0f), GMaxRHIFeatureLevel);
-	FCanvasTileItem TileItem(InPosition, InSource->GetResource(), InSize, FLinearColor::White);
-	TileItem.BatchedElementParameters = BatchedElementParameters;
+	FCanvasTileItem TileItem(FVector2D::ZeroVector, InSource->GetResource(), FLinearColor::White);
 
-	Canvas.PushAbsoluteTransform(InTransform);
+	TileItem.BatchedElementParameters = BatchedElementParameters;
 	TileItem.Draw(&Canvas);
-	Canvas.PopTransform();
 
 	Canvas.Flush_GameThread();
 	FlushRenderingCommands();
@@ -217,18 +179,6 @@ bool FGLTFTextureUtility::DrawTexture(UTextureRenderTarget2D* OutTarget, const U
 	return true;
 }
 
-bool FGLTFTextureUtility::RotateTexture(UTextureRenderTarget2D* OutTarget, const UTexture2D* InSource, const FVector2D& InPosition, const FVector2D& InSize, float InDegrees)
-{
-	FMatrix Transform = FMatrix::Identity;
-	if (InDegrees != 0)
-	{
-		const FVector Center = FVector(InSize.X / 2.0f, InSize.Y / 2.0f, 0);
-		Transform = FTranslationMatrix(-Center) * FRotationMatrix({ 0, InDegrees, 0 }) * FTranslationMatrix(Center);
-	}
-
-	return DrawTexture(OutTarget, InSource, InPosition, InSize, Transform);
-}
-
 bool FGLTFTextureUtility::ReadPixels(const UTextureRenderTarget2D* InRenderTarget, TArray<FColor>& OutPixels)
 {
 	FTextureRenderTarget2DResource* Resource = static_cast<FTextureRenderTarget2DResource*>(const_cast<UTextureRenderTarget2D*>(InRenderTarget)->GetResource());
diff --git a/Engine/Plugins/Enterprise/GLTFExporter/Source/GLTFExporter/Private/Converters/GLTFTextureUtility.h b/Engine/Plugins/Enterprise/GLTFExporter/Source/GLTFExporter/Private/Converters/GLTFTextureUtility.h
--- a/Engine/Plugins/Enterprise/GLTFExporter/Source/GLTFExporter/Private/Converters/GLTFTextureUtility.h
+++ b/Engine/Plugins/Enterprise/GLTFExporter/Source/GLTFExporter/Private/Converters/GLTFTextureUtility.h
@@ -2,15 +2,12 @@
 
 #pragma once
 
-#include "Json/GLTFJsonEnums.h"
-#include "Engine/Texture.h"
 #include "PixelFormat.h"
-#include "RHIDefinitions.h"
+#include "Engine/TextureDefines.h"
 
-class UTextureCube;
-class UTextureRenderTarget2D;
-class UTextureRenderTargetCube;
+class UTexture;
 class UTexture2D;
+class UTextureRenderTarget2D;
 
 struct FGLTFTextureUtility
 {
@@ -21,8 +18,6 @@ struct FGLTFTextureUtility
 	static bool IsHDR(const UTexture* Texture);
 	static bool IsCubemap(const UTexture* Texture);
 
-	static float GetCubeFaceRotation(ECubeFace CubeFace);
-
 	static TextureFilter GetDefaultFilter(TextureGroup Group);
 
 	static int32 GetMipBias(const UTexture* Texture);
@@ -31,12 +26,9 @@ struct FGLTFTextureUtility
 
 	static TTuple<TextureAddress, TextureAddress> GetAddressXY(const UTexture* Texture);
 
-	static UTexture2D* CreateTransientTexture(const void* RawData, int64 ByteLength, const FIntPoint& Size, EPixelFormat Format, bool bSRGB = false);
-
 	static UTextureRenderTarget2D* CreateRenderTarget(const FIntPoint& Size, bool bIsHDR);
 
-	static bool DrawTexture(UTextureRenderTarget2D* OutTarget, const UTexture2D* InSource, const FVector2D& InPosition, const FVector2D& InSize, const FMatrix& InTransform = FMatrix::Identity);
-	static bool RotateTexture(UTextureRenderTarget2D* OutTarget, const UTexture2D* InSource, const FVector2D& InPosition, const FVector2D& InSize, float InDegrees);
+	static bool DrawTexture(UTextureRenderTarget2D* OutTarget, const UTexture2D* InSource);
 
 	static bool ReadPixels(const UTextureRenderTarget2D* InRenderTarget, TArray<FColor>& OutPixels);
 
diff --git a/Engine/Plugins/Enterprise/GLTFExporter/Source/GLTFExporter/Private/Tasks/GLTFDelayedTextureTasks.cpp b/Engine/Plugins/Enterprise/GLTFExporter/Source/GLTFExporter/Private/Tasks/GLTFDelayedTextureTasks.cpp
--- a/Engine/Plugins/Enterprise/GLTFExporter/Source/GLTFExporter/Private/Tasks/GLTFDelayedTextureTasks.cpp
+++ b/Engine/Plugins/Enterprise/GLTFExporter/Source/GLTFExporter/Private/Tasks/GLTFDelayedTextureTasks.cpp
@@ -4,12 +4,8 @@
 #include "Builders/GLTFContainerBuilder.h"
 #include "Utilities/GLTFCoreUtilities.h"
 #include "Converters/GLTFTextureUtility.h"
-#include "Converters/GLTFNameUtility.h"
 #include "Engine/Texture2D.h"
-#include "Engine/TextureCube.h"
 #include "Engine/TextureRenderTarget2D.h"
-#include "Engine/TextureRenderTargetCube.h"
-#include "Engine/LightMapTexture2D.h"
 
 FString FGLTFDelayedTexture2DTask::GetName()
 {
@@ -27,12 +23,11 @@ void FGLTFDelayedTexture2DTask::Process()
 		JsonTexture->Name += bToSRGB ? TEXT("_sRGB") : TEXT("_Linear");
 	}
 
-	const bool bIsHDR = FGLTFTextureUtility::IsHDR(Texture2D);
 	const FIntPoint Size = FGLTFTextureUtility::GetInGameSize(Texture2D);
-	UTextureRenderTarget2D* RenderTarget = FGLTFTextureUtility::CreateRenderTarget(Size, bIsHDR);
+	UTextureRenderTarget2D* RenderTarget = FGLTFTextureUtility::CreateRenderTarget(Size, false);
 
 	// TODO: preserve maximum image quality (avoid compression artifacts) by copying source data (and adjustments) to a temp texture
-	FGLTFTextureUtility::DrawTexture(RenderTarget, Texture2D, FVector2D::ZeroVector, Size);
+	FGLTFTextureUtility::DrawTexture(RenderTarget, Texture2D);
 
 	TGLTFSharedArray<FColor> Pixels;
 	if (!FGLTFTextureUtility::ReadPixels(RenderTarget, *Pixels))
@@ -50,9 +45,7 @@ void FGLTFDelayedTexture2DTask::Process()
 	FGLTFTextureUtility::TransformColorSpace(*Pixels, bFromSRGB, bToSRGB);
 
 	const bool bIgnoreAlpha = FGLTFTextureUtility::IsAlphaless(Texture2D->GetPixelFormat());
-	const EGLTFTextureType Type =
-		Texture2D->IsNormalMap() ? EGLTFTextureType::Normalmaps :
-		bIsHDR ? EGLTFTextureType::HDR : EGLTFTextureType::None;
+	const EGLTFTextureType Type = Texture2D->IsNormalMap() ? EGLTFTextureType::Normalmaps : EGLTFTextureType::None;
 
 	JsonTexture->Source = Builder.AddUniqueImage(Pixels, Size, bIgnoreAlpha, Type, JsonTexture->Name);
 	JsonTexture->Sampler = Builder.AddUniqueSampler(Texture2D);
diff --git a/Engine/Plugins/Enterprise/GLTFExporter/Source/GLTFExporter/Public/Options/GLTFExportOptions.h b/Engine/Plugins/Enterprise/GLTFExporter/Source/GLTFExporter/Public/Options/GLTFExportOptions.h
--- a/Engine/Plugins/Enterprise/GLTFExporter/Source/GLTFExporter/Public/Options/GLTFExportOptions.h
+++ b/Engine/Plugins/Enterprise/GLTFExporter/Source/GLTFExporter/Public/Options/GLTFExportOptions.h
@@ -30,15 +30,6 @@ enum class EGLTFTextureType : uint8
 };
 ENUM_CLASS_FLAGS(EGLTFTextureType);
 
-UENUM(BlueprintType)
-enum class EGLTFTextureHDREncoding : uint8
-{
-	/** Clamp HDR colors to standard 8-bit per channel. */
-	None,
-	/** Encode HDR colors to RGBM (will discard alpha). */
-	RGBM
-};
-
 UENUM(BlueprintType, Meta = (Bitflags, UseEnumValuesAsMaskValuesInEditor = "true"))
 enum class EGLTFSceneMobility : uint8
 {
