Remove UGLTFExportOptions::NoLossyImageFormatFor since its more or less irrelevant without custom extensions

diff --git a/Engine/Plugins/Enterprise/GLTFExporter/Source/GLTFExporter/Private/Builders/GLTFConvertBuilder.cpp b/Engine/Plugins/Enterprise/GLTFExporter/Source/GLTFExporter/Private/Builders/GLTFConvertBuilder.cpp
--- a/Engine/Plugins/Enterprise/GLTFExporter/Source/GLTFExporter/Private/Builders/GLTFConvertBuilder.cpp
+++ b/Engine/Plugins/Enterprise/GLTFExporter/Source/GLTFExporter/Private/Builders/GLTFConvertBuilder.cpp
@@ -300,9 +300,9 @@ FGLTFJsonTexture* FGLTFConvertBuilder::AddUniqueTexture(const UTextureRenderTarg
 	return TextureRenderTarget2DConverter->GetOrAdd(Texture, bToSRGB);
 }
 
-FGLTFJsonImage* FGLTFConvertBuilder::AddUniqueImage(TGLTFSharedArray<FColor>& Pixels, FIntPoint Size, bool bIgnoreAlpha, EGLTFTextureType Type, const FString& Name)
+FGLTFJsonImage* FGLTFConvertBuilder::AddUniqueImage(TGLTFSharedArray<FColor>& Pixels, FIntPoint Size, bool bIgnoreAlpha, const FString& Name)
 {
-	return ImageConverter->GetOrAdd(Name, Type, bIgnoreAlpha, Size, Pixels);
+	return ImageConverter->GetOrAdd(Name, bIgnoreAlpha, Size, Pixels);
 }
 
 FGLTFJsonSkin* FGLTFConvertBuilder::AddUniqueSkin(FGLTFJsonNode* RootNode, const USkeletalMesh* SkeletalMesh)
diff --git a/Engine/Plugins/Enterprise/GLTFExporter/Source/GLTFExporter/Private/Converters/GLTFImageConverters.cpp b/Engine/Plugins/Enterprise/GLTFExporter/Source/GLTFExporter/Private/Converters/GLTFImageConverters.cpp
--- a/Engine/Plugins/Enterprise/GLTFExporter/Source/GLTFExporter/Private/Converters/GLTFImageConverters.cpp
+++ b/Engine/Plugins/Enterprise/GLTFExporter/Source/GLTFExporter/Private/Converters/GLTFImageConverters.cpp
@@ -5,11 +5,11 @@
 #include "Builders/GLTFConvertBuilder.h"
 #include "Misc/FileHelper.h"
 
-FGLTFJsonImage* FGLTFImageConverter::Convert(TGLTFSuperfluous<FString> Name, EGLTFTextureType Type, bool bIgnoreAlpha, FIntPoint Size, TGLTFSharedArray<FColor> Pixels)
+FGLTFJsonImage* FGLTFImageConverter::Convert(TGLTFSuperfluous<FString> Name, bool bIgnoreAlpha, FIntPoint Size, TGLTFSharedArray<FColor> Pixels)
 {
 	const TSharedPtr<FGLTFMemoryArchive> CompressedData = MakeShared<FGLTFMemoryArchive>();
 
-	const EGLTFJsonMimeType MimeType = GetMimeType(Pixels->GetData(), Size, bIgnoreAlpha, Type);
+	const EGLTFJsonMimeType MimeType = GetMimeType(Pixels->GetData(), Size, bIgnoreAlpha);
 	switch (MimeType)
 	{
 		case EGLTFJsonMimeType::None:
@@ -45,24 +45,15 @@ FGLTFJsonImage* FGLTFImageConverter::Convert(TGLTFSuperfluous<FString> Name, EGL
 	return JsonImage;
 }
 
-EGLTFJsonMimeType FGLTFImageConverter::GetMimeType(const FColor* Pixels, FIntPoint Size, bool bIgnoreAlpha, EGLTFTextureType Type) const
+EGLTFJsonMimeType FGLTFImageConverter::GetMimeType(const FColor* Pixels, FIntPoint Size, bool bIgnoreAlpha) const
 {
 	switch (Builder.ExportOptions->TextureImageFormat)
 	{
-		case EGLTFTextureImageFormat::None:
-			return EGLTFJsonMimeType::None;
-
-		case EGLTFTextureImageFormat::PNG:
-			return EGLTFJsonMimeType::PNG;
-
-		case EGLTFTextureImageFormat::JPEG:
-			return
-				(Type == EGLTFTextureType::None || !EnumHasAllFlags(static_cast<EGLTFTextureType>(Builder.ExportOptions->NoLossyImageFormatFor), Type)) &&
-				(bIgnoreAlpha || FGLTFImageUtility::NoAlphaNeeded(Pixels, Size)) ?
-				EGLTFJsonMimeType::JPEG : EGLTFJsonMimeType::PNG;
-
+		case EGLTFTextureImageFormat::None: return EGLTFJsonMimeType::None;
+		case EGLTFTextureImageFormat::PNG: return EGLTFJsonMimeType::PNG;
+		case EGLTFTextureImageFormat::JPEG: return bIgnoreAlpha || FGLTFImageUtility::NoAlphaNeeded(Pixels, Size) ? EGLTFJsonMimeType::JPEG : EGLTFJsonMimeType::PNG;
 		default:
 			checkNoEntry();
-		return EGLTFJsonMimeType::None;
+			return EGLTFJsonMimeType::None;
 	}
 }
diff --git a/Engine/Plugins/Enterprise/GLTFExporter/Source/GLTFExporter/Private/Converters/GLTFMaterialProxyFactory.cpp b/Engine/Plugins/Enterprise/GLTFExporter/Source/GLTFExporter/Private/Converters/GLTFMaterialProxyFactory.cpp
--- a/Engine/Plugins/Enterprise/GLTFExporter/Source/GLTFExporter/Private/Converters/GLTFMaterialProxyFactory.cpp
+++ b/Engine/Plugins/Enterprise/GLTFExporter/Source/GLTFExporter/Private/Converters/GLTFMaterialProxyFactory.cpp
@@ -269,13 +269,13 @@ TUniquePtr<IGLTFImageConverter> FGLTFMaterialProxyFactory::CreateImageConverter(
 
 	protected:
 
-		virtual FGLTFJsonImage* Convert(TGLTFSuperfluous<FString> Name, EGLTFTextureType Type, bool bIgnoreAlpha, FIntPoint Size, TGLTFSharedArray<FColor> Pixels) override
+		virtual FGLTFJsonImage* Convert(TGLTFSuperfluous<FString> Name, bool bIgnoreAlpha, FIntPoint Size, TGLTFSharedArray<FColor> Pixels) override
 		{
 			const FString Filename = FGLTFImageUtility::GetUniqueFilename(Name, TEXT(""), UniqueFilenames);
 			UniqueFilenames.Add(Filename);
 
 			FGLTFJsonImage* Image = Factory.Builder.AddImage();
-			Factory.Images.Add(Image, { Filename, Type, bIgnoreAlpha, Size, Pixels });
+			Factory.Images.Add(Image, { Filename, bIgnoreAlpha, Size, Pixels });
 			return Image;
 		}
 	};
diff --git a/Engine/Plugins/Enterprise/GLTFExporter/Source/GLTFExporter/Private/Converters/GLTFMaterialProxyFactory.h b/Engine/Plugins/Enterprise/GLTFExporter/Source/GLTFExporter/Private/Converters/GLTFMaterialProxyFactory.h
--- a/Engine/Plugins/Enterprise/GLTFExporter/Source/GLTFExporter/Private/Converters/GLTFMaterialProxyFactory.h
+++ b/Engine/Plugins/Enterprise/GLTFExporter/Source/GLTFExporter/Private/Converters/GLTFMaterialProxyFactory.h
@@ -27,7 +27,6 @@ private:
 	struct FGLTFImageData
 	{
 		FString Filename;
-		EGLTFTextureType Type;
 		bool bIgnoreAlpha;
 		FIntPoint Size;
 		TGLTFSharedArray<FColor> Pixels;
diff --git a/Engine/Plugins/Enterprise/GLTFExporter/Source/GLTFExporter/Private/Converters/GLTFMaterialUtility.cpp b/Engine/Plugins/Enterprise/GLTFExporter/Source/GLTFExporter/Private/Converters/GLTFMaterialUtility.cpp
--- a/Engine/Plugins/Enterprise/GLTFExporter/Source/GLTFExporter/Private/Converters/GLTFMaterialUtility.cpp
+++ b/Engine/Plugins/Enterprise/GLTFExporter/Source/GLTFExporter/Private/Converters/GLTFMaterialUtility.cpp
@@ -183,7 +183,7 @@ FGLTFJsonTexture* FGLTFMaterialUtility::AddTexture(FGLTFConvertBuilder& Builder,
 	FGLTFJsonTexture* JsonTexture = Builder.AddTexture();
 	JsonTexture->Name = TextureName;
 	JsonTexture->Sampler = JsonSampler;
-	JsonTexture->Source = Builder.AddUniqueImage(Pixels, TextureSize, bIgnoreAlpha, bIsNormalMap ? EGLTFTextureType::Normalmaps : EGLTFTextureType::None, TextureName);
+	JsonTexture->Source = Builder.AddUniqueImage(Pixels, TextureSize, bIgnoreAlpha, TextureName);
 
 	return JsonTexture;
 }
diff --git a/Engine/Plugins/Enterprise/GLTFExporter/Source/GLTFExporter/Private/Options/GLTFExportOptions.cpp b/Engine/Plugins/Enterprise/GLTFExporter/Source/GLTFExporter/Private/Options/GLTFExportOptions.cpp
--- a/Engine/Plugins/Enterprise/GLTFExporter/Source/GLTFExporter/Private/Options/GLTFExportOptions.cpp
+++ b/Engine/Plugins/Enterprise/GLTFExporter/Source/GLTFExporter/Private/Options/GLTFExportOptions.cpp
@@ -30,7 +30,6 @@ void UGLTFExportOptions::ResetToDefault()
 	bExportAnimationSequences = true;
 	TextureImageFormat = EGLTFTextureImageFormat::PNG;
 	TextureImageQuality = 0;
-	NoLossyImageFormatFor = static_cast<int32>(EGLTFTextureType::All);
 	bExportTextureTransforms = true;
 	bAdjustNormalmaps = true;
 	bExportHiddenInGame = false;
diff --git a/Engine/Plugins/Enterprise/GLTFExporter/Source/GLTFExporter/Private/Tasks/GLTFDelayedTextureTasks.cpp b/Engine/Plugins/Enterprise/GLTFExporter/Source/GLTFExporter/Private/Tasks/GLTFDelayedTextureTasks.cpp
--- a/Engine/Plugins/Enterprise/GLTFExporter/Source/GLTFExporter/Private/Tasks/GLTFDelayedTextureTasks.cpp
+++ b/Engine/Plugins/Enterprise/GLTFExporter/Source/GLTFExporter/Private/Tasks/GLTFDelayedTextureTasks.cpp
@@ -45,9 +45,8 @@ void FGLTFDelayedTexture2DTask::Process()
 	FGLTFTextureUtility::TransformColorSpace(*Pixels, bFromSRGB, bToSRGB);
 
 	const bool bIgnoreAlpha = FGLTFTextureUtility::IsAlphaless(Texture2D->GetPixelFormat());
-	const EGLTFTextureType Type = Texture2D->IsNormalMap() ? EGLTFTextureType::Normalmaps : EGLTFTextureType::None;
 
-	JsonTexture->Source = Builder.AddUniqueImage(Pixels, Size, bIgnoreAlpha, Type, JsonTexture->Name);
+	JsonTexture->Source = Builder.AddUniqueImage(Pixels, Size, bIgnoreAlpha, JsonTexture->Name);
 	JsonTexture->Sampler = Builder.AddUniqueSampler(Texture2D);
 }
 
@@ -67,9 +66,6 @@ void FGLTFDelayedTextureRenderTarget2DTask::Process()
 		JsonTexture->Name += bToSRGB ? TEXT("_sRGB") : TEXT("_Linear");
 	}
 
-	const bool bIsHDR = FGLTFTextureUtility::IsHDR(RenderTarget2D);
-	const FIntPoint Size = { RenderTarget2D->SizeX, RenderTarget2D->SizeY };
-
 	TGLTFSharedArray<FColor> Pixels;
 	if (!FGLTFTextureUtility::ReadPixels(RenderTarget2D, *Pixels))
 	{
@@ -79,9 +75,9 @@ void FGLTFDelayedTextureRenderTarget2DTask::Process()
 
 	FGLTFTextureUtility::TransformColorSpace(*Pixels, bFromSRGB, bToSRGB);
 
+	const FIntPoint Size = { RenderTarget2D->SizeX, RenderTarget2D->SizeY };
 	const bool bIgnoreAlpha = FGLTFTextureUtility::IsAlphaless(RenderTarget2D->GetFormat());
-	const EGLTFTextureType Type = bIsHDR ? EGLTFTextureType::HDR : EGLTFTextureType::None;
 
-	JsonTexture->Source = Builder.AddUniqueImage(Pixels, Size, bIgnoreAlpha, Type, JsonTexture->Name);
+	JsonTexture->Source = Builder.AddUniqueImage(Pixels, Size, bIgnoreAlpha, JsonTexture->Name);
 	JsonTexture->Sampler = Builder.AddUniqueSampler(RenderTarget2D);
 }
diff --git a/Engine/Plugins/Enterprise/GLTFExporter/Source/GLTFExporter/Public/Builders/GLTFConvertBuilder.h b/Engine/Plugins/Enterprise/GLTFExporter/Source/GLTFExporter/Public/Builders/GLTFConvertBuilder.h
--- a/Engine/Plugins/Enterprise/GLTFExporter/Source/GLTFExporter/Public/Builders/GLTFConvertBuilder.h
+++ b/Engine/Plugins/Enterprise/GLTFExporter/Source/GLTFExporter/Public/Builders/GLTFConvertBuilder.h
@@ -64,7 +64,7 @@ public:
 	FGLTFJsonTexture* AddUniqueTexture(const UTexture* Texture, bool bToSRGB);
 	FGLTFJsonTexture* AddUniqueTexture(const UTexture2D* Texture, bool bToSRGB);
 	FGLTFJsonTexture* AddUniqueTexture(const UTextureRenderTarget2D* Texture, bool bToSRGB);
-	FGLTFJsonImage* AddUniqueImage(TGLTFSharedArray<FColor>& Pixels, FIntPoint Size, bool bIgnoreAlpha, EGLTFTextureType Type, const FString& Name);
+	FGLTFJsonImage* AddUniqueImage(TGLTFSharedArray<FColor>& Pixels, FIntPoint Size, bool bIgnoreAlpha, const FString& Name);
 
 	FGLTFJsonSkin* AddUniqueSkin(FGLTFJsonNode* RootNode, const USkeletalMesh* SkeletalMesh);
 	FGLTFJsonSkin* AddUniqueSkin(FGLTFJsonNode* RootNode, const USkeletalMeshComponent* SkeletalMeshComponent);
diff --git a/Engine/Plugins/Enterprise/GLTFExporter/Source/GLTFExporter/Public/Converters/GLTFImageConverters.h b/Engine/Plugins/Enterprise/GLTFExporter/Source/GLTFExporter/Public/Converters/GLTFImageConverters.h
--- a/Engine/Plugins/Enterprise/GLTFExporter/Source/GLTFExporter/Public/Converters/GLTFImageConverters.h
+++ b/Engine/Plugins/Enterprise/GLTFExporter/Source/GLTFExporter/Public/Converters/GLTFImageConverters.h
@@ -7,9 +7,8 @@
 #include "Converters/GLTFBuilderContext.h"
 #include "Converters/GLTFSuperfluous.h"
 #include "Converters/GLTFSharedArray.h"
-#include "Options/GLTFExportOptions.h"
 
-typedef TGLTFConverter<FGLTFJsonImage*, TGLTFSuperfluous<FString>, EGLTFTextureType, bool, FIntPoint, TGLTFSharedArray<FColor>> IGLTFImageConverter;
+typedef TGLTFConverter<FGLTFJsonImage*, TGLTFSuperfluous<FString>, bool, FIntPoint, TGLTFSharedArray<FColor>> IGLTFImageConverter;
 
 class GLTFEXPORTER_API FGLTFImageConverter : public FGLTFBuilderContext, public IGLTFImageConverter
 {
@@ -19,9 +18,9 @@ public:
 
 protected:
 
-	virtual FGLTFJsonImage* Convert(TGLTFSuperfluous<FString> Name, EGLTFTextureType Type, bool bIgnoreAlpha, FIntPoint Size, TGLTFSharedArray<FColor> Pixels) override;
+	virtual FGLTFJsonImage* Convert(TGLTFSuperfluous<FString> Name, bool bIgnoreAlpha, FIntPoint Size, TGLTFSharedArray<FColor> Pixels) override;
 
 private:
 
-	EGLTFJsonMimeType GetMimeType(const FColor* Pixels, FIntPoint Size, bool bIgnoreAlpha, EGLTFTextureType Type) const;
+	EGLTFJsonMimeType GetMimeType(const FColor* Pixels, FIntPoint Size, bool bIgnoreAlpha) const;
 };
diff --git a/Engine/Plugins/Enterprise/GLTFExporter/Source/GLTFExporter/Public/Options/GLTFExportOptions.h b/Engine/Plugins/Enterprise/GLTFExporter/Source/GLTFExporter/Public/Options/GLTFExportOptions.h
--- a/Engine/Plugins/Enterprise/GLTFExporter/Source/GLTFExporter/Public/Options/GLTFExportOptions.h
+++ b/Engine/Plugins/Enterprise/GLTFExporter/Source/GLTFExporter/Public/Options/GLTFExportOptions.h
@@ -17,19 +17,6 @@ enum class EGLTFTextureImageFormat : uint8
 	JPEG UMETA(DisplayName = "JPEG (if no alpha)")
 };
 
-UENUM(BlueprintType, Meta = (Bitflags, UseEnumValuesAsMaskValuesInEditor = "true"))
-enum class EGLTFTextureType : uint8
-{
-	None = 0 UMETA(Hidden),
-
-	HDR = 1 << 0,
-	Normalmaps = 1 << 1,
-	Lightmaps = 1 << 2,
-
-	All = HDR | Normalmaps | Lightmaps UMETA(Hidden)
-};
-ENUM_CLASS_FLAGS(EGLTFTextureType);
-
 UENUM(BlueprintType, Meta = (Bitflags, UseEnumValuesAsMaskValuesInEditor = "true"))
 enum class EGLTFSceneMobility : uint8
 {
