Remove UGLTFExportOptions::bIncludeGeneratorVersion since its use has been made obsolete with improvements in gltf exporter tests plugin/framework

diff --git a/Engine/Plugins/Enterprise/GLTFExporter/Source/GLTFExporter/Private/Builders/GLTFJsonBuilder.cpp b/Engine/Plugins/Enterprise/GLTFExporter/Source/GLTFExporter/Private/Builders/GLTFJsonBuilder.cpp
--- a/Engine/Plugins/Enterprise/GLTFExporter/Source/GLTFExporter/Private/Builders/GLTFJsonBuilder.cpp
+++ b/Engine/Plugins/Enterprise/GLTFExporter/Source/GLTFExporter/Private/Builders/GLTFJsonBuilder.cpp
@@ -1,7 +1,6 @@
 // Copyright Epic Games, Inc. All Rights Reserved.
 
 #include "Builders/GLTFJsonBuilder.h"
-#include "GLTFExporterModule.h"
 #include "Interfaces/IPluginManager.h"
 #include "Runtime/Launch/Resources/Version.h"
 
@@ -106,12 +105,7 @@ const FGLTFJsonRoot& FGLTFJsonBuilder::GetRoot() const
 	return JsonRoot;
 }
 
-FString FGLTFJsonBuilder::GetGeneratorString() const
+const TCHAR* FGLTFJsonBuilder::GetGeneratorString()
 {
-	const TSharedPtr<IPlugin> Plugin = IPluginManager::Get().FindPlugin(GLTFEXPORTER_MODULE_NAME);
-	const FPluginDescriptor& PluginDescriptor = Plugin->GetDescriptor();
-
-	return ExportOptions->bIncludeGeneratorVersion
-		? TEXT(EPIC_PRODUCT_NAME) TEXT(" ") ENGINE_VERSION_STRING TEXT(" ") + PluginDescriptor.FriendlyName + TEXT(" ") + PluginDescriptor.VersionName
-		: TEXT(EPIC_PRODUCT_NAME) TEXT(" ") + PluginDescriptor.FriendlyName;
+	return TEXT(EPIC_PRODUCT_NAME) TEXT(" ") ENGINE_VERSION_STRING;
 }
diff --git a/Engine/Plugins/Enterprise/GLTFExporter/Source/GLTFExporter/Private/Options/GLTFExportOptions.cpp b/Engine/Plugins/Enterprise/GLTFExporter/Source/GLTFExporter/Private/Options/GLTFExportOptions.cpp
--- a/Engine/Plugins/Enterprise/GLTFExporter/Source/GLTFExporter/Private/Options/GLTFExportOptions.cpp
+++ b/Engine/Plugins/Enterprise/GLTFExporter/Source/GLTFExporter/Private/Options/GLTFExportOptions.cpp
@@ -14,7 +14,6 @@ void UGLTFExportOptions::ResetToDefault()
 	bExportPreviewMesh = true;
 	bStrictCompliance = true;
 	bSkipNearDefaultValues = true;
-	bIncludeGeneratorVersion = true;
 	bExportProxyMaterials = true;
 	bExportUnlitMaterials = true;
 	bExportClearCoatMaterials = true;
diff --git a/Engine/Plugins/Enterprise/GLTFExporter/Source/GLTFExporter/Public/Builders/GLTFJsonBuilder.h b/Engine/Plugins/Enterprise/GLTFExporter/Source/GLTFExporter/Public/Builders/GLTFJsonBuilder.h
--- a/Engine/Plugins/Enterprise/GLTFExporter/Source/GLTFExporter/Public/Builders/GLTFJsonBuilder.h
+++ b/Engine/Plugins/Enterprise/GLTFExporter/Source/GLTFExporter/Public/Builders/GLTFJsonBuilder.h
@@ -39,7 +39,7 @@ protected:
 
 private:
 
-	FString GetGeneratorString() const;
+	static const TCHAR* GetGeneratorString();
 
 	FGLTFJsonRoot JsonRoot;
 };
diff --git a/Engine/Plugins/Enterprise/GLTFExporter/Source/GLTFExporter/Public/Options/GLTFExportOptions.h b/Engine/Plugins/Enterprise/GLTFExporter/Source/GLTFExporter/Public/Options/GLTFExportOptions.h
--- a/Engine/Plugins/Enterprise/GLTFExporter/Source/GLTFExporter/Public/Options/GLTFExportOptions.h
+++ b/Engine/Plugins/Enterprise/GLTFExporter/Source/GLTFExporter/Public/Options/GLTFExportOptions.h
@@ -62,10 +62,6 @@ class GLTFEXPORTER_API UGLTFExportOptions : public UObject
 	UPROPERTY(EditAnywhere, BlueprintReadWrite, Config, Category = General)
 	bool bSkipNearDefaultValues;
 
-	/** If enabled, version info for Unreal Engine and exporter plugin will be included as metadata in the glTF asset, which is useful when reporting issues. */
-	UPROPERTY(EditAnywhere, BlueprintReadWrite, Config, Category = General)
-	bool bIncludeGeneratorVersion;
-
 	/** If enabled, materials that have a proxy defined in their user data, will be exported using that proxy instead. This setting won't affect proxy materials exported or referenced directly. */
 	UPROPERTY(EditAnywhere, BlueprintReadWrite, Config, Category = Material)
 	bool bExportProxyMaterials;
