Fix incorrect conversion of light intensity when exporting to glTF and intenisty units set to non-candelas

diff --git a/Engine/Plugins/Enterprise/GLTFExporter/Source/GLTFExporter/Private/Converters/GLTFLightConverters.cpp b/Engine/Plugins/Enterprise/GLTFExporter/Source/GLTFExporter/Private/Converters/GLTFLightConverters.cpp
--- a/Engine/Plugins/Enterprise/GLTFExporter/Source/GLTFExporter/Private/Converters/GLTFLightConverters.cpp
+++ b/Engine/Plugins/Enterprise/GLTFExporter/Source/GLTFExporter/Private/Converters/GLTFLightConverters.cpp
@@ -21,10 +21,13 @@ FGLTFJsonLight* FGLTFLightConverter::Convert(const ULightComponent* LightCompone
 	Light->Name = FGLTFNameUtility::GetName(LightComponent);
 	Light->Type = Type;
 
-	Light->Intensity = LightComponent->Intensity;
+	const float ConversionScale = Type == EGLTFJsonLightType::Directional ? 1.0f : (0.01f * 0.01f); // Conversion from cm2 to m2
+	const FLinearColor ColorBrightness = LightComponent->GetColoredLightBrightness() * ConversionScale;
+	const float Brightness = FMath::Max(ColorBrightness.GetMax(), 1.0f);
+	const FLinearColor Color = ColorBrightness / Brightness;
 
-	const FLinearColor TemperatureColor = LightComponent->bUseTemperature ? FLinearColor::MakeFromColorTemperature(LightComponent->Temperature) : FLinearColor::White;
-	Light->Color = FGLTFCoreUtilities::ConvertColor3(TemperatureColor * LightComponent->GetLightColor(), Builder.ExportOptions->bStrictCompliance);
+	Light->Intensity = Brightness;
+	Light->Color = FGLTFCoreUtilities::ConvertColor3(Color, false);
 
 	if (const UPointLightComponent* PointLightComponent = Cast<UPointLightComponent>(LightComponent))
 	{
