Remove UGLTFExportOptions::bStrictCompliance since it should really always be enabled to keep consistency across all asset complexities and avoid potential issues in different viewers

diff --git a/Engine/Plugins/Enterprise/GLTFExporter/Source/GLTFExporter/Private/Converters/GLTFLightConverters.cpp b/Engine/Plugins/Enterprise/GLTFExporter/Source/GLTFExporter/Private/Converters/GLTFLightConverters.cpp
--- a/Engine/Plugins/Enterprise/GLTFExporter/Source/GLTFExporter/Private/Converters/GLTFLightConverters.cpp
+++ b/Engine/Plugins/Enterprise/GLTFExporter/Source/GLTFExporter/Private/Converters/GLTFLightConverters.cpp
@@ -27,7 +27,7 @@ FGLTFJsonLight* FGLTFLightConverter::Convert(const ULightComponent* LightCompone
 	const FLinearColor Color = ColorBrightness / Brightness;
 
 	Light->Intensity = Brightness;
-	Light->Color = FGLTFCoreUtilities::ConvertColor3(Color, false);
+	Light->Color = FGLTFCoreUtilities::ConvertColor3(Color);
 
 	if (const UPointLightComponent* PointLightComponent = Cast<UPointLightComponent>(LightComponent))
 	{
@@ -39,11 +39,10 @@ FGLTFJsonLight* FGLTFLightConverter::Convert(const ULightComponent* LightCompone
 		Light->Spot.InnerConeAngle = FGLTFCoreUtilities::ConvertLightAngle(SpotLightComponent->InnerConeAngle);
 		Light->Spot.OuterConeAngle = FGLTFCoreUtilities::ConvertLightAngle(SpotLightComponent->OuterConeAngle);
 
-		if (Builder.ExportOptions->bStrictCompliance)
-		{
-			Light->Spot.InnerConeAngle = FMath::Clamp(Light->Spot.InnerConeAngle, 0.0f, nextafterf(Light->Spot.OuterConeAngle, 0.0f));
-			Light->Spot.OuterConeAngle = FMath::Clamp(Light->Spot.OuterConeAngle, nextafterf(Light->Spot.InnerConeAngle, HALF_PI), HALF_PI);
-		}
+		// KHR_lights_punctual requires that InnerConeAngle must be greater than or equal to 0 and less than OuterConeAngle.
+		Light->Spot.InnerConeAngle = FMath::Clamp(Light->Spot.InnerConeAngle, 0.0f, nextafterf(Light->Spot.OuterConeAngle, 0.0f));
+		// KHR_lights_punctual requires that OuterConeAngle must be greater than InnerConeAngle and less than or equal to PI / 2.0.
+		Light->Spot.OuterConeAngle = FMath::Clamp(Light->Spot.OuterConeAngle, nextafterf(Light->Spot.InnerConeAngle, HALF_PI), HALF_PI);
 	}
 
 	return Light;
diff --git a/Engine/Plugins/Enterprise/GLTFExporter/Source/GLTFExporter/Private/Options/GLTFExportOptions.cpp b/Engine/Plugins/Enterprise/GLTFExporter/Source/GLTFExporter/Private/Options/GLTFExportOptions.cpp
--- a/Engine/Plugins/Enterprise/GLTFExporter/Source/GLTFExporter/Private/Options/GLTFExportOptions.cpp
+++ b/Engine/Plugins/Enterprise/GLTFExporter/Source/GLTFExporter/Private/Options/GLTFExportOptions.cpp
@@ -12,7 +12,6 @@ void UGLTFExportOptions::ResetToDefault()
 {
 	ExportUniformScale = 0.01;
 	bExportPreviewMesh = true;
-	bStrictCompliance = true;
 	bSkipNearDefaultValues = true;
 	bExportProxyMaterials = true;
 	bExportUnlitMaterials = true;
diff --git a/Engine/Plugins/Enterprise/GLTFExporter/Source/GLTFExporter/Private/Tasks/GLTFDelayedMaterialTasks.cpp b/Engine/Plugins/Enterprise/GLTFExporter/Source/GLTFExporter/Private/Tasks/GLTFDelayedMaterialTasks.cpp
--- a/Engine/Plugins/Enterprise/GLTFExporter/Source/GLTFExporter/Private/Tasks/GLTFDelayedMaterialTasks.cpp
+++ b/Engine/Plugins/Enterprise/GLTFExporter/Source/GLTFExporter/Private/Tasks/GLTFDelayedMaterialTasks.cpp
@@ -255,7 +255,7 @@ void FGLTFDelayedMaterialTask::GetProxyParameter(const TGLTFProxyMaterialParamet
 	FLinearColor Value;
 	if (ParameterInfo.Get(Material, Value, true))
 	{
-		OutValue = FGLTFCoreUtilities::ConvertColor3(Value, false);
+		OutValue = FGLTFCoreUtilities::ConvertColor3(Value);
 	}
 }
 
@@ -264,7 +264,7 @@ void FGLTFDelayedMaterialTask::GetProxyParameter(const TGLTFProxyMaterialParamet
 	FLinearColor Value;
 	if (ParameterInfo.Get(Material, Value, true))
 	{
-		OutValue = FGLTFCoreUtilities::ConvertColor(Value, false);
+		OutValue = FGLTFCoreUtilities::ConvertColor(Value);
 	}
 }
 
@@ -523,7 +523,7 @@ bool FGLTFDelayedMaterialTask::TryGetBaseColorAndOpacity(FGLTFJsonPBRMetallicRou
 		FLinearColor BaseColorFactor(BaseColorBakeOutput.ConstantValue * BaseColorScale);
 		BaseColorFactor.A = OpacityBakeOutput.ConstantValue.R;
 
-		OutPBRParams.BaseColorFactor = FGLTFCoreUtilities::ConvertColor(BaseColorFactor, Builder.ExportOptions->bStrictCompliance);
+		OutPBRParams.BaseColorFactor = FGLTFCoreUtilities::ConvertColor(BaseColorFactor);
 		return true;
 	}
 
@@ -573,7 +573,9 @@ bool FGLTFDelayedMaterialTask::TryGetBaseColorAndOpacity(FGLTFJsonPBRMetallicRou
 
 	OutPBRParams.BaseColorTexture.TexCoord = CombinedTexCoord;
 	OutPBRParams.BaseColorTexture.Index = CombinedTexture;
-	OutPBRParams.BaseColorFactor = FGLTFCoreUtilities::ConvertColor({ BaseColorScale, BaseColorScale, BaseColorScale }, Builder.ExportOptions->bStrictCompliance);
+
+	// TODO: add support for KHR_materials_emissive_strength
+	OutPBRParams.BaseColorFactor = FGLTFCoreUtilities::ConvertColor(FLinearColor(BaseColorScale, BaseColorScale, BaseColorScale));
 
 	return true;
 }
@@ -855,7 +857,7 @@ bool FGLTFDelayedMaterialTask::TryGetEmissive(FGLTFJsonMaterial& OutMaterial, co
 	if (PropertyBakeOutput.bIsConstant)
 	{
 		const FLinearColor EmissiveColor = PropertyBakeOutput.ConstantValue;
-		OutMaterial.EmissiveFactor = FGLTFCoreUtilities::ConvertColor3(EmissiveColor * EmissiveScale, Builder.ExportOptions->bStrictCompliance);
+		OutMaterial.EmissiveFactor = FGLTFCoreUtilities::ConvertColor3(EmissiveColor * EmissiveScale);
 	}
 	else
 	{
@@ -906,7 +908,7 @@ bool FGLTFDelayedMaterialTask::TryGetConstantColor(FGLTFJsonColor3& OutValue, co
 	FLinearColor Value;
 	if (TryGetConstantColor(Value, Property))
 	{
-		OutValue = FGLTFCoreUtilities::ConvertColor3(Value, Builder.ExportOptions->bStrictCompliance);
+		OutValue = FGLTFCoreUtilities::ConvertColor3(Value);
 		return true;
 	}
 
@@ -918,7 +920,7 @@ bool FGLTFDelayedMaterialTask::TryGetConstantColor(FGLTFJsonColor4& OutValue, co
 	FLinearColor Value;
 	if (TryGetConstantColor(Value, Property))
 	{
-		OutValue = FGLTFCoreUtilities::ConvertColor(Value, Builder.ExportOptions->bStrictCompliance);
+		OutValue = FGLTFCoreUtilities::ConvertColor(Value);
 		return true;
 	}
 
@@ -1288,7 +1290,7 @@ bool FGLTFDelayedMaterialTask::TryGetBakedMaterialProperty(FGLTFJsonTextureInfo&
 
 	if (PropertyBakeOutput.bIsConstant)
 	{
-		OutConstant = FGLTFCoreUtilities::ConvertColor3(PropertyBakeOutput.ConstantValue, Builder.ExportOptions->bStrictCompliance);
+		OutConstant = FGLTFCoreUtilities::ConvertColor3(PropertyBakeOutput.ConstantValue);
 		return true;
 	}
 
@@ -1322,7 +1324,7 @@ bool FGLTFDelayedMaterialTask::TryGetBakedMaterialProperty(FGLTFJsonTextureInfo&
 
 	if (PropertyBakeOutput.bIsConstant)
 	{
-		OutConstant = FGLTFCoreUtilities::ConvertColor(PropertyBakeOutput.ConstantValue, Builder.ExportOptions->bStrictCompliance);
+		OutConstant = FGLTFCoreUtilities::ConvertColor(PropertyBakeOutput.ConstantValue);
 		return true;
 	}
 
diff --git a/Engine/Plugins/Enterprise/GLTFExporter/Source/GLTFExporter/Private/Utilities/GLTFCoreUtilities.cpp b/Engine/Plugins/Enterprise/GLTFExporter/Source/GLTFExporter/Private/Utilities/GLTFCoreUtilities.cpp
--- a/Engine/Plugins/Enterprise/GLTFExporter/Source/GLTFExporter/Private/Utilities/GLTFCoreUtilities.cpp
+++ b/Engine/Plugins/Enterprise/GLTFExporter/Source/GLTFExporter/Private/Utilities/GLTFCoreUtilities.cpp
@@ -67,43 +67,24 @@ FGLTFVector2 FGLTFCoreUtilities::ConvertUV(const FVector2DHalf& UV)
 	return ConvertUV(FVector2f(UV));
 }
 
-FGLTFColor4 FGLTFCoreUtilities::ConvertColor(const FLinearColor& Color, bool bForceLDR)
+FGLTFColor4 FGLTFCoreUtilities::ConvertColor(const FLinearColor& Color)
 {
-	if (bForceLDR)
-	{
-		return {
-			FMath::Clamp(Color.R, 0.0f, 1.0f),
-			FMath::Clamp(Color.G, 0.0f, 1.0f),
-			FMath::Clamp(Color.B, 0.0f, 1.0f),
-			FMath::Clamp(Color.A, 0.0f, 1.0f)
-		};
-	}
-
-	// Just make sure its non-negative (which can happen when using MakeFromColorTemperature).
+	// glTF requires that color components are in the range 0.0 - 1.0
 	return {
-		FMath::Max(Color.R, 0.0f),
-		FMath::Max(Color.G, 0.0f),
-		FMath::Max(Color.B, 0.0f),
-		FMath::Max(Color.A, 0.0f)
+		FMath::Clamp(Color.R, 0.0f, 1.0f),
+		FMath::Clamp(Color.G, 0.0f, 1.0f),
+		FMath::Clamp(Color.B, 0.0f, 1.0f),
+		FMath::Clamp(Color.A, 0.0f, 1.0f)
 	};
 }
 
-FGLTFColor3 FGLTFCoreUtilities::ConvertColor3(const FLinearColor& Color, bool bForceLDR)
+FGLTFColor3 FGLTFCoreUtilities::ConvertColor3(const FLinearColor& Color)
 {
-	if (bForceLDR)
-	{
-		return {
-			FMath::Clamp(Color.R, 0.0f, 1.0f),
-			FMath::Clamp(Color.G, 0.0f, 1.0f),
-			FMath::Clamp(Color.B, 0.0f, 1.0f)
-		};
-	}
-
-	// Just make sure its non-negative (which can happen when using MakeFromColorTemperature).
+	// glTF requires that color components are in the range 0.0 - 1.0
 	return {
-		FMath::Max(Color.R, 0.0f),
-		FMath::Max(Color.G, 0.0f),
-		FMath::Max(Color.B, 0.0f)
+		FMath::Clamp(Color.R, 0.0f, 1.0f),
+		FMath::Clamp(Color.G, 0.0f, 1.0f),
+		FMath::Clamp(Color.B, 0.0f, 1.0f)
 	};
 }
 
diff --git a/Engine/Plugins/Enterprise/GLTFExporter/Source/GLTFExporter/Public/Options/GLTFExportOptions.h b/Engine/Plugins/Enterprise/GLTFExporter/Source/GLTFExporter/Public/Options/GLTFExportOptions.h
--- a/Engine/Plugins/Enterprise/GLTFExporter/Source/GLTFExporter/Public/Options/GLTFExportOptions.h
+++ b/Engine/Plugins/Enterprise/GLTFExporter/Source/GLTFExporter/Public/Options/GLTFExportOptions.h
@@ -54,10 +54,6 @@ class GLTFEXPORTER_API UGLTFExportOptions : public UObject
 	UPROPERTY(EditAnywhere, BlueprintReadWrite, Config, Category = General)
 	bool bExportPreviewMesh;
 
-	/** If enabled, certain values (like HDR colors and light angles) will be truncated during export to strictly conform to the formal glTF specification. */
-	UPROPERTY(EditAnywhere, BlueprintReadWrite, Config, Category = General)
-	bool bStrictCompliance;
-
 	/** If enabled, floating-point-based JSON properties that are nearly equal to their default value will not be exported and thus regarded as exactly default, reducing JSON size. */
 	UPROPERTY(EditAnywhere, BlueprintReadWrite, Config, Category = General)
 	bool bSkipNearDefaultValues;
diff --git a/Engine/Plugins/Enterprise/GLTFExporter/Source/GLTFExporter/Public/Utilities/GLTFCoreUtilities.h b/Engine/Plugins/Enterprise/GLTFExporter/Source/GLTFExporter/Public/Utilities/GLTFCoreUtilities.h
--- a/Engine/Plugins/Enterprise/GLTFExporter/Source/GLTFExporter/Public/Utilities/GLTFCoreUtilities.h
+++ b/Engine/Plugins/Enterprise/GLTFExporter/Source/GLTFExporter/Public/Utilities/GLTFCoreUtilities.h
@@ -36,8 +36,8 @@ struct GLTFEXPORTER_API FGLTFCoreUtilities
 	static FGLTFVector2 ConvertUV(const FVector2f& UV);
 	static FGLTFVector2 ConvertUV(const FVector2DHalf& UV);
 
-	static FGLTFColor4 ConvertColor(const FLinearColor& Color, bool bForceLDR);
-	static FGLTFColor3 ConvertColor3(const FLinearColor& Color, bool bForceLDR);
+	static FGLTFColor4 ConvertColor(const FLinearColor& Color);
+	static FGLTFColor3 ConvertColor3(const FLinearColor& Color);
 	static FGLTFUInt8Color4 ConvertColor(const FColor& Color);
 
 	static FGLTFQuaternion ConvertRotation(const FRotator3f& Rotation);
