Remove unused UGLTFExportOptions::NoLossyImageFormatFor

diff --git a/Engine/Plugins/Enterprise/GLTFExporter/Source/GLTFExporter/Public/Options/GLTFExportOptions.h b/Engine/Plugins/Enterprise/GLTFExporter/Source/GLTFExporter/Public/Options/GLTFExportOptions.h
--- a/Engine/Plugins/Enterprise/GLTFExporter/Source/GLTFExporter/Public/Options/GLTFExportOptions.h
+++ b/Engine/Plugins/Enterprise/GLTFExporter/Source/GLTFExporter/Public/Options/GLTFExportOptions.h
@@ -122,10 +122,6 @@ class GLTFEXPORTER_API UGLTFExportOptions : public UObject
 	UPROPERTY(EditAnywhere, BlueprintReadWrite, Config, Category = Texture, Meta = (ClampMin = "0", ClampMax = "100", EditCondition = "TextureImageFormat == EGLTFTextureImageFormat::JPEG"))
 	int32 TextureImageQuality;
 
-	/** Texture types that will always use lossless formats (e.g. PNG) because of sensitivity to compression artifacts. */
-	UPROPERTY(EditAnywhere, BlueprintReadWrite, Config, Category = Texture, Meta = (Bitmask, BitmaskEnum = "/Script/GLTFExporter.EGLTFTextureType", EditCondition = "TextureImageFormat == EGLTFTextureImageFormat::JPEG"))
-	int32 NoLossyImageFormatFor; // Bitmask combined from EGLTFTextureType
-
 	/** If enabled, export UV tiling and un-mirroring settings in a texture coordinate expression node for simple material input expressions. Uses extension KHR_texture_transform. */
 	UPROPERTY(EditAnywhere, BlueprintReadWrite, Config, Category = Texture, Meta = (EditCondition = "TextureImageFormat != EGLTFTextureImageFormat::None"))
 	bool bExportTextureTransforms;
