Replace UGLTFExportOptions::ExportLights (int32) with UGLTFExportOptions::bExportLights (bool) since mobility mask is no longer relevant without the custom gltf extension for lightmaps

diff --git a/Engine/Plugins/Enterprise/GLTFExporter/Source/GLTFExporter/Private/Builders/GLTFBuilder.cpp b/Engine/Plugins/Enterprise/GLTFExporter/Source/GLTFExporter/Private/Builders/GLTFBuilder.cpp
--- a/Engine/Plugins/Enterprise/GLTFExporter/Source/GLTFExporter/Private/Builders/GLTFBuilder.cpp
+++ b/Engine/Plugins/Enterprise/GLTFExporter/Source/GLTFExporter/Private/Builders/GLTFBuilder.cpp
@@ -76,13 +76,6 @@ TextureAddress FGLTFBuilder::GetBakeTilingForMaterialProperty(const UMaterialInt
 	return UGLTFMaterialExportOptions::GetBakeTilingForPropertyGroup(Material, PropertyGroup, DefaultValue);
 }
 
-bool FGLTFBuilder::ShouldExportLight(EComponentMobility::Type LightMobility) const
-{
-	const EGLTFSceneMobility AllowedMobility = static_cast<EGLTFSceneMobility>(ExportOptions->ExportLights);
-	const EGLTFSceneMobility QueriedMobility = GetSceneMobility(LightMobility);
-	return EnumHasAllFlags(AllowedMobility, QueriedMobility);
-}
-
 int32 FGLTFBuilder::SanitizeLOD(const UStaticMesh* StaticMesh, const UStaticMeshComponent* StaticMeshComponent, int32 LODIndex) const
 {
 	return LODIndex > 0
@@ -120,16 +113,3 @@ const UGLTFExportOptions* FGLTFBuilder::SanitizeExportOptions(const UGLTFExportO
 
 	return Options;
 }
-
-EGLTFSceneMobility FGLTFBuilder::GetSceneMobility(EComponentMobility::Type Mobility)
-{
-	switch (Mobility)
-	{
-		case EComponentMobility::Static:     return EGLTFSceneMobility::Static;
-		case EComponentMobility::Stationary: return EGLTFSceneMobility::Stationary;
-		case EComponentMobility::Movable:    return EGLTFSceneMobility::Movable;
-		default:
-			checkNoEntry();
-			return EGLTFSceneMobility::None;
-	}
-}
diff --git a/Engine/Plugins/Enterprise/GLTFExporter/Source/GLTFExporter/Private/Converters/GLTFNodeConverters.cpp b/Engine/Plugins/Enterprise/GLTFExporter/Source/GLTFExporter/Private/Converters/GLTFNodeConverters.cpp
--- a/Engine/Plugins/Enterprise/GLTFExporter/Source/GLTFExporter/Private/Converters/GLTFNodeConverters.cpp
+++ b/Engine/Plugins/Enterprise/GLTFExporter/Source/GLTFExporter/Private/Converters/GLTFNodeConverters.cpp
@@ -164,7 +164,7 @@ void FGLTFComponentConverter::ConvertComponentSpecialization(const USceneCompone
 	}
 	else if (const ULightComponent* LightComponent = Cast<ULightComponent>(SceneComponent))
 	{
-		if (Builder.ShouldExportLight(LightComponent->Mobility))
+		if (Builder.ExportOptions->bExportLights)
 		{
 			// TODO: conversion of light direction should be done in separate converter
 			FGLTFJsonNode* LightNode = Builder.AddNode();
diff --git a/Engine/Plugins/Enterprise/GLTFExporter/Source/GLTFExporter/Private/Options/GLTFExportOptions.cpp b/Engine/Plugins/Enterprise/GLTFExporter/Source/GLTFExporter/Private/Options/GLTFExportOptions.cpp
--- a/Engine/Plugins/Enterprise/GLTFExporter/Source/GLTFExporter/Private/Options/GLTFExportOptions.cpp
+++ b/Engine/Plugins/Enterprise/GLTFExporter/Source/GLTFExporter/Private/Options/GLTFExportOptions.cpp
@@ -31,7 +31,7 @@ void UGLTFExportOptions::ResetToDefault()
 	bExportTextureTransforms = true;
 	bAdjustNormalmaps = true;
 	bExportHiddenInGame = false;
-	ExportLights = static_cast<int32>(EGLTFSceneMobility::Stationary | EGLTFSceneMobility::Movable);
+	bExportLights = true;
 	bExportCameras = true;
 	ExportMaterialVariants = EGLTFMaterialVariantMode::UseMeshData;
 }
diff --git a/Engine/Plugins/Enterprise/GLTFExporter/Source/GLTFExporter/Public/Builders/GLTFBuilder.h b/Engine/Plugins/Enterprise/GLTFExporter/Source/GLTFExporter/Public/Builders/GLTFBuilder.h
--- a/Engine/Plugins/Enterprise/GLTFExporter/Source/GLTFExporter/Public/Builders/GLTFBuilder.h
+++ b/Engine/Plugins/Enterprise/GLTFExporter/Source/GLTFExporter/Public/Builders/GLTFBuilder.h
@@ -4,7 +4,6 @@
 
 #include "Engine/EngineTypes.h"
 #include "Options/GLTFExportOptions.h"
-#include "Json/GLTFJsonEnums.h"
 #include "UObject/GCObjectScopeGuard.h"
 
 class UMaterialInterface;
@@ -31,8 +30,6 @@ public:
 	TextureFilter GetBakeFilterForMaterialProperty(const UMaterialInterface* Material, EGLTFMaterialPropertyGroup PropertyGroup) const;
 	TextureAddress GetBakeTilingForMaterialProperty(const UMaterialInterface* Material, EGLTFMaterialPropertyGroup PropertyGroup) const;
 
-	bool ShouldExportLight(EComponentMobility::Type LightMobility) const;
-
 	int32 SanitizeLOD(const UStaticMesh* StaticMesh, const UStaticMeshComponent* StaticMeshComponent, int32 LODIndex) const;
 	int32 SanitizeLOD(const USkeletalMesh* SkeletalMesh, const USkeletalMeshComponent* SkeletalMeshComponent, int32 LODIndex) const;
 
@@ -41,5 +38,4 @@ private:
 	FGCObjectScopeGuard ExportOptionsGuard;
 
 	static const UGLTFExportOptions* SanitizeExportOptions(const UGLTFExportOptions* Options);
-	static EGLTFSceneMobility GetSceneMobility(EComponentMobility::Type Mobility);
 };
diff --git a/Engine/Plugins/Enterprise/GLTFExporter/Source/GLTFExporter/Public/Options/GLTFExportOptions.h b/Engine/Plugins/Enterprise/GLTFExporter/Source/GLTFExporter/Public/Options/GLTFExportOptions.h
--- a/Engine/Plugins/Enterprise/GLTFExporter/Source/GLTFExporter/Public/Options/GLTFExportOptions.h
+++ b/Engine/Plugins/Enterprise/GLTFExporter/Source/GLTFExporter/Public/Options/GLTFExportOptions.h
@@ -17,19 +17,6 @@ enum class EGLTFTextureImageFormat : uint8
 	JPEG UMETA(DisplayName = "JPEG (if no alpha)")
 };
 
-UENUM(BlueprintType, Meta = (Bitflags, UseEnumValuesAsMaskValuesInEditor = "true"))
-enum class EGLTFSceneMobility : uint8
-{
-	None = 0 UMETA(Hidden),
-
-	Static = 1 << 0,
-	Stationary = 1 << 1,
-	Movable = 1 << 2,
-
-	All = Static | Stationary | Movable UMETA(Hidden)
-};
-ENUM_CLASS_FLAGS(EGLTFSceneMobility);
-
 UENUM(BlueprintType)
 enum class EGLTFMaterialVariantMode : uint8
 {
@@ -134,9 +121,9 @@ class GLTFEXPORTER_API UGLTFExportOptions : public UObject
 	UPROPERTY(EditAnywhere, BlueprintReadWrite, Config, Category = Scene)
 	bool bExportHiddenInGame;
 
-	/** Mobility of directional, point, and spot light components that will be exported. Uses extension KHR_lights_punctual. */
-	UPROPERTY(EditAnywhere, BlueprintReadWrite, Config, Category = Scene, Meta = (Bitmask, BitmaskEnum = "/Script/GLTFExporter.EGLTFSceneMobility"))
-	int32 ExportLights; // Bitmask combined from EGLTFSceneMobility
+	/** If enabled, export directional, point, and spot light components. Uses extension KHR_lights_punctual. */
+	UPROPERTY(EditAnywhere, BlueprintReadWrite, Config, Category = Scene)
+	bool bExportLights;
 
 	/** If enabled, export camera components. */
 	UPROPERTY(EditAnywhere, BlueprintReadWrite, Config, Category = Scene)
