Remove old (no longer needed) hack for accessing vertex skin weight lookup data

diff --git a/Engine/Plugins/Enterprise/GLTFExporter/Source/GLTFExporter/Private/Converters/GLTFBufferAdapter.cpp b/Engine/Plugins/Enterprise/GLTFExporter/Source/GLTFExporter/Private/Converters/GLTFBufferAdapter.cpp
--- a/Engine/Plugins/Enterprise/GLTFExporter/Source/GLTFExporter/Private/Converters/GLTFBufferAdapter.cpp
+++ b/Engine/Plugins/Enterprise/GLTFExporter/Source/GLTFExporter/Private/Converters/GLTFBufferAdapter.cpp
@@ -95,7 +95,7 @@ TUniquePtr<IGLTFBufferAdapter> IGLTFBufferAdapter::GetInfluences(const FSkinWeig
 TUniquePtr<IGLTFBufferAdapter> IGLTFBufferAdapter::GetLookups(const FSkinWeightVertexBuffer* VertexBuffer)
 {
 	const FSkinWeightLookupVertexBuffer* LookupBuffer = VertexBuffer->GetLookupVertexBuffer();
-	const void* LookupData = FGLTFBufferUtility::GetCPUBuffer(LookupBuffer);
+	const void* LookupData = LookupBuffer->GetLookupData();
 	if (LookupData != nullptr && FGLTFBufferUtility::HasCPUAccess(VertexBuffer)) return MakeUnique<FGLTFBufferAdapterCPU>(LookupData);
 	return MakeUnique<FGLTFBufferAdapterGPU>(LookupBuffer->VertexBufferRHI);
 }
diff --git a/Engine/Plugins/Enterprise/GLTFExporter/Source/GLTFExporter/Private/Converters/GLTFBufferUtility.cpp b/Engine/Plugins/Enterprise/GLTFExporter/Source/GLTFExporter/Private/Converters/GLTFBufferUtility.cpp
--- a/Engine/Plugins/Enterprise/GLTFExporter/Source/GLTFExporter/Private/Converters/GLTFBufferUtility.cpp
+++ b/Engine/Plugins/Enterprise/GLTFExporter/Source/GLTFExporter/Private/Converters/GLTFBufferUtility.cpp
@@ -67,21 +67,6 @@ const void* FGLTFBufferUtility::GetCPUBuffer(const FRawStaticIndexBuffer16or32In
 	return IndexBuffer->GetResourceDataSize() > 0 ? const_cast<FRawStaticIndexBuffer16or32Interface*>(IndexBuffer)->GetPointerTo(0) : nullptr;
 }
 
-const void* FGLTFBufferUtility::GetCPUBuffer(const FSkinWeightLookupVertexBuffer* VertexBuffer)
-{
-	struct FSkinWeightLookupVertexBufferHack : FVertexBuffer
-	{
-		FShaderResourceViewRHIRef SRVValue;
-		bool bNeedsCPUAccess;
-		FStaticMeshVertexDataInterface* LookupData;
-		uint8* Data;
-		uint32 NumVertices;
-	};
-
-	static_assert(sizeof(FSkinWeightLookupVertexBufferHack) == sizeof(FSkinWeightLookupVertexBuffer), "FSkinWeightLookupVertexBufferHack memory layout doesn't match FSkinWeightLookupVertexBuffer");
-	return reinterpret_cast<const FSkinWeightLookupVertexBufferHack*>(VertexBuffer)->Data;
-}
-
 void FGLTFBufferUtility::ReadRHIBuffer(FRHIBuffer* SourceBuffer, TArray<uint8>& OutData)
 {
 	OutData.Empty();
diff --git a/Engine/Plugins/Enterprise/GLTFExporter/Source/GLTFExporter/Private/Converters/GLTFBufferUtility.h b/Engine/Plugins/Enterprise/GLTFExporter/Source/GLTFExporter/Private/Converters/GLTFBufferUtility.h
--- a/Engine/Plugins/Enterprise/GLTFExporter/Source/GLTFExporter/Private/Converters/GLTFBufferUtility.h
+++ b/Engine/Plugins/Enterprise/GLTFExporter/Source/GLTFExporter/Private/Converters/GLTFBufferUtility.h
@@ -25,7 +25,6 @@ struct FGLTFBufferUtility
 
 	static const void* GetCPUBuffer(const FRawStaticIndexBuffer* IndexBuffer);
 	static const void* GetCPUBuffer(const FRawStaticIndexBuffer16or32Interface* IndexBuffer);
-	static const void* GetCPUBuffer(const FSkinWeightLookupVertexBuffer* VertexBuffer);
 
 	static void ReadRHIBuffer(FRHIBuffer* SourceBuffer, TArray<uint8>& OutData);
 };
