Add default value (empty set) to param SelectedActors in UFUNCTION ExportToGLTF

diff --git a/Engine/Plugins/Enterprise/GLTFExporter/Source/GLTFExporter/Public/Exporters/GLTFExporter.h b/Engine/Plugins/Enterprise/GLTFExporter/Source/GLTFExporter/Public/Exporters/GLTFExporter.h
--- a/Engine/Plugins/Enterprise/GLTFExporter/Source/GLTFExporter/Public/Exporters/GLTFExporter.h
+++ b/Engine/Plugins/Enterprise/GLTFExporter/Source/GLTFExporter/Public/Exporters/GLTFExporter.h
@@ -46,7 +46,7 @@ public:
 	 *
 	 * @return true if the object was successfully exported
 	 */
-	UFUNCTION(BlueprintCallable, Category = "Miscellaneous")
+	UFUNCTION(BlueprintCallable, Category = "Miscellaneous", meta=(AutoCreateRefTerm="SelectedActors"))
 	static bool ExportToGLTF(UObject* Object, const FString& FilePath, const UGLTFExportOptions* Options, const TSet<AActor*>& SelectedActors, FGLTFExportMessages& OutMessages);
 
 	/**
