Use actor labels instead of actor names (unless runtime) for root components. Non-root components use their own names.

diff --git a/Engine/Plugins/Enterprise/GLTFExporter/Source/GLTFExporter/Private/Converters/GLTFNameUtility.cpp b/Engine/Plugins/Enterprise/GLTFExporter/Source/GLTFExporter/Private/Converters/GLTFNameUtility.cpp
--- a/Engine/Plugins/Enterprise/GLTFExporter/Source/GLTFExporter/Private/Converters/GLTFNameUtility.cpp
+++ b/Engine/Plugins/Enterprise/GLTFExporter/Source/GLTFExporter/Private/Converters/GLTFNameUtility.cpp
@@ -1,55 +1,29 @@
 // Copyright Epic Games, Inc. All Rights Reserved.
 
 #include "Converters/GLTFNameUtility.h"
-#include "Engine/StaticMeshActor.h"
-#include "Animation/SkeletalMeshActor.h"
-#include "Components/LightComponent.h"
-#include "Components/SkeletalMeshComponent.h"
-#include "Components/StaticMeshComponent.h"
-#include "Engine/Light.h"
-#include "Camera/CameraComponent.h"
-#include "Camera/CameraActor.h"
+#include "Components/SceneComponent.h"
+#include "GameFramework/Actor.h"
 
-namespace
+FString FGLTFNameUtility::GetName(const UEnum* Enum, int32 Value)
 {
-	template <typename ActorType>
-	FString GetActorNameIfOwnerOfType(const UActorComponent* Component)
-	{
-		if (const ActorType* Owner = Cast<ActorType>(Component->GetOwner()))
-		{
-			return Owner->GetName();
-		}
-
-		return Component->GetName();
-	}
+	check(Enum != nullptr);
+	const FString DisplayName = Enum->GetDisplayNameTextByValue(Value).ToString();
+	return DisplayName.IsEmpty() ? FString::FromInt(Value) : DisplayName;
 }
 
 FString FGLTFNameUtility::GetName(const USceneComponent* Component)
 {
 	if (const AActor* Owner = Component->GetOwner())
 	{
-		return Owner->GetName() + TEXT("_") + Component->GetName();
+		if (Component == Owner->GetRootComponent())
+		{
+#if WITH_EDITOR
+			return Owner->GetActorLabel();
+#else
+			return Owner->GetName();
+#endif
+		}
 	}
 
 	return Component->GetName();
 }
-
-FString FGLTFNameUtility::GetName(const UStaticMeshComponent* Component)
-{
-	return GetActorNameIfOwnerOfType<AStaticMeshActor>(Component);
-}
-
-FString FGLTFNameUtility::GetName(const USkeletalMeshComponent* Component)
-{
-	return GetActorNameIfOwnerOfType<ASkeletalMeshActor>(Component);
-}
-
-FString FGLTFNameUtility::GetName(const ULightComponent* Component)
-{
-	return GetActorNameIfOwnerOfType<ALight>(Component);
-}
-
-FString FGLTFNameUtility::GetName(const UCameraComponent* Component)
-{
-	return GetActorNameIfOwnerOfType<ACameraActor>(Component);
-}
diff --git a/Engine/Plugins/Enterprise/GLTFExporter/Source/GLTFExporter/Private/Converters/GLTFNameUtility.h b/Engine/Plugins/Enterprise/GLTFExporter/Source/GLTFExporter/Private/Converters/GLTFNameUtility.h
--- a/Engine/Plugins/Enterprise/GLTFExporter/Source/GLTFExporter/Private/Converters/GLTFNameUtility.h
+++ b/Engine/Plugins/Enterprise/GLTFExporter/Source/GLTFExporter/Private/Converters/GLTFNameUtility.h
@@ -3,34 +3,19 @@
 #pragma once
 
 #include "CoreMinimal.h"
-#include "UObject/Class.h"
 #include "UObject/ReflectedTypeAccessors.h"
 
-class ULightComponent;
-class UCameraComponent;
-
-class UCameraComponent;
-class ULightComponent;
 class USceneComponent;
-class USkeletalMeshComponent;
-class UStaticMeshComponent;
 
 struct FGLTFNameUtility
 {
 	template <typename EnumType, typename = typename TEnableIf<TIsEnum<EnumType>::Value>::Type>
 	static FString GetName(EnumType Value)
 	{
-		const UEnum* Enum = StaticEnum<EnumType>();
-		check(Enum != nullptr);
-		const int32 NumericValue = static_cast<int32>(Value);
-		const FString DisplayName = Enum->GetDisplayNameTextByValue(NumericValue).ToString();
-		return DisplayName.IsEmpty() ? FString::FromInt(NumericValue) : DisplayName;
+		return GetName(StaticEnum<EnumType>(), static_cast<int32>(Value));
 	}
 
-	static FString GetName(const USceneComponent* Component);
+	static FString GetName(const UEnum* Enum, int32 Value);
 
-	static FString GetName(const UStaticMeshComponent* Component);
-	static FString GetName(const USkeletalMeshComponent* Component);
-	static FString GetName(const ULightComponent* Component);
-	static FString GetName(const UCameraComponent* Component);
+	static FString GetName(const USceneComponent* Component);
 };
