Fix horde issue 237310

diff --git a/Engine/Plugins/Enterprise/GLTFExporter/Source/GLTFExporter/Private/Tasks/GLTFDelayedMeshTasks.cpp b/Engine/Plugins/Enterprise/GLTFExporter/Source/GLTFExporter/Private/Tasks/GLTFDelayedMeshTasks.cpp
--- a/Engine/Plugins/Enterprise/GLTFExporter/Source/GLTFExporter/Private/Tasks/GLTFDelayedMeshTasks.cpp
+++ b/Engine/Plugins/Enterprise/GLTFExporter/Source/GLTFExporter/Private/Tasks/GLTFDelayedMeshTasks.cpp
@@ -109,6 +109,11 @@ namespace
 	}
 }
 
+FString FGLTFDelayedStaticMeshTask::GetName()
+{
+	return StaticMeshComponent != nullptr ? FGLTFNameUtility::GetName(StaticMeshComponent) : StaticMesh->GetName();
+}
+
 void FGLTFDelayedStaticMeshTask::Process()
 {
 	FGLTFMeshUtility::FullyLoad(StaticMesh);
@@ -207,6 +212,11 @@ void FGLTFDelayedStaticMeshTask::Process()
 	}
 }
 
+FString FGLTFDelayedSkeletalMeshTask::GetName()
+{
+	return SkeletalMeshComponent != nullptr ? FGLTFNameUtility::GetName(SkeletalMeshComponent) : SkeletalMesh->GetName();
+}
+
 void FGLTFDelayedSkeletalMeshTask::Process()
 {
 	FGLTFMeshUtility::FullyLoad(SkeletalMesh);
diff --git a/Engine/Plugins/Enterprise/GLTFExporter/Source/GLTFExporter/Private/Tasks/GLTFDelayedMeshTasks.h b/Engine/Plugins/Enterprise/GLTFExporter/Source/GLTFExporter/Private/Tasks/GLTFDelayedMeshTasks.h
--- a/Engine/Plugins/Enterprise/GLTFExporter/Source/GLTFExporter/Private/Tasks/GLTFDelayedMeshTasks.h
+++ b/Engine/Plugins/Enterprise/GLTFExporter/Source/GLTFExporter/Private/Tasks/GLTFDelayedMeshTasks.h
@@ -29,10 +29,7 @@ public:
 	{
 	}
 
-	virtual FString GetName() override
-	{
-		return StaticMeshComponent != nullptr ? FGLTFNameUtility::GetName(StaticMeshComponent) : StaticMesh->GetName();
-	}
+	virtual FString GetName() override;
 
 	virtual void Process() override;
 
@@ -63,10 +60,7 @@ public:
 	{
 	}
 
-	virtual FString GetName() override
-	{
-		return SkeletalMeshComponent != nullptr ? FGLTFNameUtility::GetName(SkeletalMeshComponent) : SkeletalMesh->GetName();
-	}
+	virtual FString GetName() override;
 
 	virtual void Process() override;
 
