Avoid exporting identical samplers by decoupling export of sampler from texture input

diff --git a/Engine/Plugins/Enterprise/GLTFExporter/Source/GLTFExporter/Private/Builders/GLTFConvertBuilder.cpp b/Engine/Plugins/Enterprise/GLTFExporter/Source/GLTFExporter/Private/Builders/GLTFConvertBuilder.cpp
--- a/Engine/Plugins/Enterprise/GLTFExporter/Source/GLTFExporter/Private/Builders/GLTFConvertBuilder.cpp
+++ b/Engine/Plugins/Enterprise/GLTFExporter/Source/GLTFExporter/Private/Builders/GLTFConvertBuilder.cpp
@@ -247,7 +247,17 @@ FGLTFJsonSampler* FGLTFConvertBuilder::AddUniqueSampler(const UTexture* Texture)
 		return nullptr;
 	}
 
-	return SamplerConverter->GetOrAdd(Texture);
+	return TextureSamplerConverter->GetOrAdd(Texture);
+}
+
+FGLTFJsonSampler* FGLTFConvertBuilder::AddUniqueSampler(TextureAddress Address, TextureFilter Filter, TextureGroup LODGroup)
+{
+	return AddUniqueSampler(Address, Address, Filter, LODGroup);
+}
+
+FGLTFJsonSampler* FGLTFConvertBuilder::AddUniqueSampler(TextureAddress AddressX, TextureAddress AddressY, TextureFilter Filter, TextureGroup LODGroup)
+{
+	return SamplerConverter->GetOrAdd(AddressX, AddressY, Filter, LODGroup);
 }
 
 FGLTFJsonTexture* FGLTFConvertBuilder::AddUniqueTexture(const UTexture* Texture)
diff --git a/Engine/Plugins/Enterprise/GLTFExporter/Source/GLTFExporter/Private/Converters/GLTFMaterialUtility.cpp b/Engine/Plugins/Enterprise/GLTFExporter/Source/GLTFExporter/Private/Converters/GLTFMaterialUtility.cpp
--- a/Engine/Plugins/Enterprise/GLTFExporter/Source/GLTFExporter/Private/Converters/GLTFMaterialUtility.cpp
+++ b/Engine/Plugins/Enterprise/GLTFExporter/Source/GLTFExporter/Private/Converters/GLTFMaterialUtility.cpp
@@ -169,20 +169,12 @@ FGLTFPropertyBakeOutput FGLTFMaterialUtility::BakeMaterialProperty(const FIntPoi
 	return PropertyBakeOutput;
 }
 
-FGLTFJsonTexture* FGLTFMaterialUtility::AddTexture(FGLTFConvertBuilder& Builder, TGLTFSharedArray<FColor>& Pixels, const FIntPoint& TextureSize, bool bIgnoreAlpha, bool bIsNormalMap, const FString& TextureName, EGLTFJsonTextureFilter MinFilter, EGLTFJsonTextureFilter MagFilter, EGLTFJsonTextureWrap WrapS, EGLTFJsonTextureWrap WrapT)
+FGLTFJsonTexture* FGLTFMaterialUtility::AddTexture(FGLTFConvertBuilder& Builder, TGLTFSharedArray<FColor>& Pixels, const FIntPoint& TextureSize, bool bIgnoreAlpha, bool bIsNormalMap, const FString& TextureName, TextureAddress TextureAddress, TextureFilter TextureFilter)
 {
-	// TODO: maybe we should reuse existing samplers?
-	FGLTFJsonSampler* JsonSampler = Builder.AddSampler();
-	JsonSampler->Name = TextureName;
-	JsonSampler->MinFilter = MinFilter;
-	JsonSampler->MagFilter = MagFilter;
-	JsonSampler->WrapS = WrapS;
-	JsonSampler->WrapT = WrapT;
-
 	// TODO: reuse same texture index when image is the same
 	FGLTFJsonTexture* JsonTexture = Builder.AddTexture();
 	JsonTexture->Name = TextureName;
-	JsonTexture->Sampler = JsonSampler;
+	JsonTexture->Sampler = Builder.AddUniqueSampler(TextureAddress, TextureFilter);
 	JsonTexture->Source = Builder.AddUniqueImage(Pixels, TextureSize, bIgnoreAlpha, TextureName);
 
 	return JsonTexture;
diff --git a/Engine/Plugins/Enterprise/GLTFExporter/Source/GLTFExporter/Private/Converters/GLTFMaterialUtility.h b/Engine/Plugins/Enterprise/GLTFExporter/Source/GLTFExporter/Private/Converters/GLTFMaterialUtility.h
--- a/Engine/Plugins/Enterprise/GLTFExporter/Source/GLTFExporter/Private/Converters/GLTFMaterialUtility.h
+++ b/Engine/Plugins/Enterprise/GLTFExporter/Source/GLTFExporter/Private/Converters/GLTFMaterialUtility.h
@@ -55,7 +55,7 @@ struct FGLTFMaterialUtility
 
 	static FGLTFPropertyBakeOutput BakeMaterialProperty(const FIntPoint& OutputSize, const FGLTFMaterialPropertyEx& Property, const UMaterialInterface* Material, int32 TexCoord, const FGLTFMeshData* MeshData = nullptr, const FGLTFIndexArray& MeshSectionIndices = {}, bool bFillAlpha = true, bool bAdjustNormalmaps = true);
 
-	static FGLTFJsonTexture* AddTexture(FGLTFConvertBuilder& Builder, TGLTFSharedArray<FColor>& Pixels, const FIntPoint& TextureSize, bool bIgnoreAlpha, bool bIsNormalMap, const FString& TextureName, EGLTFJsonTextureFilter MinFilter, EGLTFJsonTextureFilter MagFilter, EGLTFJsonTextureWrap WrapS, EGLTFJsonTextureWrap WrapT);
+	static FGLTFJsonTexture* AddTexture(FGLTFConvertBuilder& Builder, TGLTFSharedArray<FColor>& Pixels, const FIntPoint& TextureSize, bool bIgnoreAlpha, bool bIsNormalMap, const FString& TextureName, TextureAddress TextureAddress, TextureFilter TextureFilter);
 
 	static FLinearColor GetMask(const FExpressionInput& ExpressionInput);
 	static uint32 GetMaskComponentCount(const FExpressionInput& ExpressionInput);
diff --git a/Engine/Plugins/Enterprise/GLTFExporter/Source/GLTFExporter/Private/Converters/GLTFSamplerConverters.cpp b/Engine/Plugins/Enterprise/GLTFExporter/Source/GLTFExporter/Private/Converters/GLTFSamplerConverters.cpp
--- a/Engine/Plugins/Enterprise/GLTFExporter/Source/GLTFExporter/Private/Converters/GLTFSamplerConverters.cpp
+++ b/Engine/Plugins/Enterprise/GLTFExporter/Source/GLTFExporter/Private/Converters/GLTFSamplerConverters.cpp
@@ -5,29 +5,40 @@
 #include "Utilities/GLTFCoreUtilities.h"
 #include "Converters/GLTFTextureUtility.h"
 
-FGLTFJsonSampler* FGLTFSamplerConverter::Convert(const UTexture* Texture)
+FGLTFJsonSampler* FGLTFTextureSamplerConverter::Convert(const UTexture* Texture)
 {
-	// TODO: reuse existing samplers, don't create a unique sampler per texture
-
-	FGLTFJsonSampler* JsonSampler = Builder.AddSampler();
-	Texture->GetName(JsonSampler->Name);
-
-	JsonSampler->MinFilter = FGLTFCoreUtilities::ConvertMinFilter(Texture->Filter, Texture->LODGroup);
-	JsonSampler->MagFilter = FGLTFCoreUtilities::ConvertMagFilter(Texture->Filter, Texture->LODGroup);
+	TextureAddress AddressX;
+	TextureAddress AddressY;
 
 	if (FGLTFTextureUtility::IsCubemap(Texture))
 	{
-		JsonSampler->WrapS = EGLTFJsonTextureWrap::ClampToEdge;
-		JsonSampler->WrapT = EGLTFJsonTextureWrap::ClampToEdge;
+		AddressX = TA_Clamp;
+		AddressY = TA_Clamp;
 	}
 	else
 	{
-		const TTuple<TextureAddress, TextureAddress> AddressXY = FGLTFTextureUtility::GetAddressXY(Texture);
-		// TODO: report error if AddressX == TA_MAX or AddressY == TA_MAX
+		FGLTFTextureUtility::GetAddressXY(Texture, AddressX, AddressY);
+	}
+
+	return Builder.AddUniqueSampler(AddressX, AddressY, Texture->Filter, Texture->LODGroup);
+}
 
-		JsonSampler->WrapS = FGLTFCoreUtilities::ConvertWrap(AddressXY.Get<0>());
-		JsonSampler->WrapT = FGLTFCoreUtilities::ConvertWrap(AddressXY.Get<1>());
+void FGLTFSamplerConverter::Sanitize(TextureAddress& AddressX, TextureAddress& AddressY, TextureFilter& Filter, TextureGroup& LODGroup)
+{
+	if (Filter == TF_Default)
+	{
+		Filter = FGLTFTextureUtility::GetDefaultFilter(LODGroup);
 	}
 
+	LODGroup = TEXTUREGROUP_MAX; // Ignore it, since Filter should cover any missing information
+}
+
+FGLTFJsonSampler* FGLTFSamplerConverter::Convert(TextureAddress AddressX, TextureAddress AddressY, TextureFilter Filter, TextureGroup LODGroup)
+{
+	FGLTFJsonSampler* JsonSampler = Builder.AddSampler();
+	JsonSampler->MinFilter = FGLTFCoreUtilities::ConvertMinFilter(Filter);
+	JsonSampler->MagFilter = FGLTFCoreUtilities::ConvertMagFilter(Filter);
+	JsonSampler->WrapS = FGLTFCoreUtilities::ConvertWrap(AddressX);
+	JsonSampler->WrapT = FGLTFCoreUtilities::ConvertWrap(AddressY);
 	return JsonSampler;
 }
diff --git a/Engine/Plugins/Enterprise/GLTFExporter/Source/GLTFExporter/Private/Converters/GLTFTextureUtility.cpp b/Engine/Plugins/Enterprise/GLTFExporter/Source/GLTFExporter/Private/Converters/GLTFTextureUtility.cpp
--- a/Engine/Plugins/Enterprise/GLTFExporter/Source/GLTFExporter/Private/Converters/GLTFTextureUtility.cpp
+++ b/Engine/Plugins/Enterprise/GLTFExporter/Source/GLTFExporter/Private/Converters/GLTFTextureUtility.cpp
@@ -76,7 +76,7 @@ TextureFilter FGLTFTextureUtility::GetDefaultFilter(TextureGroup LODGroup)
 		case ETextureSamplerFilter::Trilinear:         return TF_Trilinear;
 		case ETextureSamplerFilter::AnisotropicPoint:  return TF_Trilinear; // A lot of engine code doesn't result in nearest
 		case ETextureSamplerFilter::AnisotropicLinear: return TF_Trilinear;
-		default:                                       return TF_Default; // Let caller decide fallback
+		default:                                       return TF_MAX;
 	}
 }
 
@@ -103,28 +103,23 @@ FIntPoint FGLTFTextureUtility::GetInGameSize(const UTexture* Texture)
 	return { InGameWidth, InGameHeight };
 }
 
-TTuple<TextureAddress, TextureAddress> FGLTFTextureUtility::GetAddressXY(const UTexture* Texture)
+void FGLTFTextureUtility::GetAddressXY(const UTexture* Texture, TextureAddress& OutAddressX, TextureAddress& OutAddressY)
 {
-	TextureAddress AddressX;
-	TextureAddress AddressY;
-
 	if (const UTexture2D* Texture2D = Cast<UTexture2D>(Texture))
 	{
-		AddressX = Texture2D->AddressX;
-		AddressY = Texture2D->AddressY;
+		OutAddressX = Texture2D->AddressX;
+		OutAddressY = Texture2D->AddressY;
 	}
 	else if (const UTextureRenderTarget2D* RenderTarget2D = Cast<UTextureRenderTarget2D>(Texture))
 	{
-		AddressX = RenderTarget2D->AddressX;
-		AddressY = RenderTarget2D->AddressY;
+		OutAddressX = RenderTarget2D->AddressX;
+		OutAddressY = RenderTarget2D->AddressY;
 	}
 	else
 	{
-		AddressX = TA_MAX;
-		AddressY = TA_MAX;
+		OutAddressX = TA_MAX;
+		OutAddressY = TA_MAX;
 	}
-
-	return MakeTuple(AddressX, AddressY);
 }
 
 UTextureRenderTarget2D* FGLTFTextureUtility::CreateRenderTarget(const FIntPoint& Size, bool bHDR)
diff --git a/Engine/Plugins/Enterprise/GLTFExporter/Source/GLTFExporter/Private/Converters/GLTFTextureUtility.h b/Engine/Plugins/Enterprise/GLTFExporter/Source/GLTFExporter/Private/Converters/GLTFTextureUtility.h
--- a/Engine/Plugins/Enterprise/GLTFExporter/Source/GLTFExporter/Private/Converters/GLTFTextureUtility.h
+++ b/Engine/Plugins/Enterprise/GLTFExporter/Source/GLTFExporter/Private/Converters/GLTFTextureUtility.h
@@ -24,7 +24,7 @@ struct FGLTFTextureUtility
 
 	static FIntPoint GetInGameSize(const UTexture* Texture);
 
-	static TTuple<TextureAddress, TextureAddress> GetAddressXY(const UTexture* Texture);
+	static void GetAddressXY(const UTexture* Texture, TextureAddress& OutAddressX, TextureAddress& OutAddressY);
 
 	static UTextureRenderTarget2D* CreateRenderTarget(const FIntPoint& Size, bool bIsHDR);
 
diff --git a/Engine/Plugins/Enterprise/GLTFExporter/Source/GLTFExporter/Private/Tasks/GLTFDelayedMaterialTasks.cpp b/Engine/Plugins/Enterprise/GLTFExporter/Source/GLTFExporter/Private/Tasks/GLTFDelayedMaterialTasks.cpp
--- a/Engine/Plugins/Enterprise/GLTFExporter/Source/GLTFExporter/Private/Tasks/GLTFDelayedMaterialTasks.cpp
+++ b/Engine/Plugins/Enterprise/GLTFExporter/Source/GLTFExporter/Private/Tasks/GLTFDelayedMaterialTasks.cpp
@@ -504,14 +504,8 @@ bool FGLTFDelayedMaterialTask::TryGetBaseColorAndOpacity(FGLTFJsonPBRMetallicRou
 	}
 
 	const FIntPoint TextureSize = Builder.GetBakeSizeForMaterialProperty(Material, GetPropertyGroup(BaseColorProperty));
-
 	const TextureAddress TextureAddress = Builder.GetBakeTilingForMaterialProperty(Material, GetPropertyGroup(BaseColorProperty));
-	const EGLTFJsonTextureWrap TextureWrapS = FGLTFCoreUtilities::ConvertWrap(TextureAddress);
-	const EGLTFJsonTextureWrap TextureWrapT = FGLTFCoreUtilities::ConvertWrap(TextureAddress);
-
 	const TextureFilter TextureFilter = Builder.GetBakeFilterForMaterialProperty(Material, GetPropertyGroup(BaseColorProperty));
-	const EGLTFJsonTextureFilter TextureMinFilter = FGLTFCoreUtilities::ConvertMinFilter(TextureFilter);
-	const EGLTFJsonTextureFilter TextureMagFilter = FGLTFCoreUtilities::ConvertMagFilter(TextureFilter);
 
 	const FGLTFPropertyBakeOutput BaseColorBakeOutput = BakeMaterialProperty(BaseColorProperty, BaseColorTexCoord, TextureSize, false);
 	const FGLTFPropertyBakeOutput OpacityBakeOutput = BakeMaterialProperty(OpacityProperty, OpacityTexCoord, TextureSize, false);
@@ -566,10 +560,8 @@ bool FGLTFDelayedMaterialTask::TryGetBaseColorAndOpacity(FGLTFJsonPBRMetallicRou
 		false,
 		false,
 		GetBakedTextureName(TEXT("BaseColor")),
-		TextureMinFilter,
-		TextureMagFilter,
-		TextureWrapS,
-		TextureWrapT);
+		TextureAddress,
+		TextureFilter);
 
 	OutPBRParams.BaseColorTexture.TexCoord = CombinedTexCoord;
 	OutPBRParams.BaseColorTexture.Index = CombinedTexture;
@@ -631,14 +623,8 @@ bool FGLTFDelayedMaterialTask::TryGetMetallicAndRoughness(FGLTFJsonPBRMetallicRo
 
 	// TODO: add support for calculating the ideal resolution to use for baking based on connected (texture) nodes
 	const FIntPoint TextureSize = Builder.GetBakeSizeForMaterialProperty(Material, GetPropertyGroup(MetallicProperty));
-
 	const TextureAddress TextureAddress = Builder.GetBakeTilingForMaterialProperty(Material, GetPropertyGroup(MetallicProperty));
-	const EGLTFJsonTextureWrap TextureWrapS = FGLTFCoreUtilities::ConvertWrap(TextureAddress);
-	const EGLTFJsonTextureWrap TextureWrapT = FGLTFCoreUtilities::ConvertWrap(TextureAddress);
-
 	const TextureFilter TextureFilter = Builder.GetBakeFilterForMaterialProperty(Material, GetPropertyGroup(MetallicProperty));
-	const EGLTFJsonTextureFilter TextureMinFilter = FGLTFCoreUtilities::ConvertMinFilter(TextureFilter);
-	const EGLTFJsonTextureFilter TextureMagFilter = FGLTFCoreUtilities::ConvertMagFilter(TextureFilter);
 
 	FGLTFPropertyBakeOutput MetallicBakeOutput = BakeMaterialProperty(MetallicProperty, MetallicTexCoord, TextureSize, false);
 	FGLTFPropertyBakeOutput RoughnessBakeOutput = BakeMaterialProperty(RoughnessProperty, RoughnessTexCoord, TextureSize, false);
@@ -690,10 +676,8 @@ bool FGLTFDelayedMaterialTask::TryGetMetallicAndRoughness(FGLTFJsonPBRMetallicRo
 		true, // NOTE: we can ignore alpha in everything but TryGetBaseColorAndOpacity
 		false,
 		GetBakedTextureName(TEXT("MetallicRoughness")),
-		TextureMinFilter,
-		TextureMagFilter,
-		TextureWrapS,
-		TextureWrapT);
+		TextureAddress,
+		TextureFilter);
 
 	OutPBRParams.MetallicRoughnessTexture.TexCoord = CombinedTexCoord;
 	OutPBRParams.MetallicRoughnessTexture.Index = CombinedTexture;
@@ -754,14 +738,8 @@ bool FGLTFDelayedMaterialTask::TryGetClearCoatRoughness(FGLTFJsonClearCoatExtens
 	}
 
 	const FIntPoint TextureSize = Builder.GetBakeSizeForMaterialProperty(Material, GetPropertyGroup(IntensityProperty));
-
 	const TextureAddress TextureAddress = Builder.GetBakeTilingForMaterialProperty(Material,GetPropertyGroup(IntensityProperty));
-	const EGLTFJsonTextureWrap TextureWrapS = FGLTFCoreUtilities::ConvertWrap(TextureAddress);
-	const EGLTFJsonTextureWrap TextureWrapT = FGLTFCoreUtilities::ConvertWrap(TextureAddress);
-
 	const TextureFilter TextureFilter = Builder.GetBakeFilterForMaterialProperty(Material, GetPropertyGroup(IntensityProperty));
-	const EGLTFJsonTextureFilter TextureMinFilter = FGLTFCoreUtilities::ConvertMinFilter(TextureFilter);
-	const EGLTFJsonTextureFilter TextureMagFilter = FGLTFCoreUtilities::ConvertMagFilter(TextureFilter);
 
 	const FGLTFPropertyBakeOutput IntensityBakeOutput = BakeMaterialProperty(IntensityProperty, IntensityTexCoord, TextureSize, false);
 	const FGLTFPropertyBakeOutput RoughnessBakeOutput = BakeMaterialProperty(RoughnessProperty, RoughnessTexCoord, TextureSize, false);
@@ -813,10 +791,8 @@ bool FGLTFDelayedMaterialTask::TryGetClearCoatRoughness(FGLTFJsonClearCoatExtens
 		true, // NOTE: we can ignore alpha in everything but TryGetBaseColorAndOpacity
 		false,
 		GetBakedTextureName(TEXT("ClearCoatRoughness")),
-		TextureMinFilter,
-		TextureMagFilter,
-		TextureWrapS,
-		TextureWrapT);
+		TextureAddress,
+		TextureFilter);
 
 	OutExtParams.ClearCoatTexture.Index = CombinedTexture;
 	OutExtParams.ClearCoatTexture.TexCoord = CombinedTexCoord;
@@ -1430,10 +1406,8 @@ bool FGLTFDelayedMaterialTask::TryGetBakedMaterialProperty(FGLTFJsonTextureInfo&
 		true, // NOTE: we can ignore alpha in everything but TryGetBaseColorAndOpacity
 		false, // Normal and ClearCoatBottomNormal are handled above
 		GetBakedTextureName(PropertyName),
-		EGLTFJsonTextureFilter::Nearest,
-		EGLTFJsonTextureFilter::Nearest,
-		EGLTFJsonTextureWrap::ClampToEdge,
-		EGLTFJsonTextureWrap::ClampToEdge);
+		TA_Clamp,
+		TF_Nearest);
 
 	OutTexInfo.Index = Texture;
 	return true;
@@ -1493,12 +1467,7 @@ FGLTFPropertyBakeOutput FGLTFDelayedMaterialTask::BakeMaterialProperty(const FGL
 bool FGLTFDelayedMaterialTask::StoreBakedPropertyTexture(FGLTFJsonTextureInfo& OutTexInfo, FGLTFPropertyBakeOutput& PropertyBakeOutput, const FString& PropertyName) const
 {
 	const TextureAddress TextureAddress = Builder.GetBakeTilingForMaterialProperty(Material, GetPropertyGroup(PropertyBakeOutput.Property));
-	const EGLTFJsonTextureWrap TextureWrapS = FGLTFCoreUtilities::ConvertWrap(TextureAddress);
-	const EGLTFJsonTextureWrap TextureWrapT = FGLTFCoreUtilities::ConvertWrap(TextureAddress);
-
 	const TextureFilter TextureFilter = Builder.GetBakeFilterForMaterialProperty(Material, GetPropertyGroup(PropertyBakeOutput.Property));
-	const EGLTFJsonTextureFilter TextureMinFilter = FGLTFCoreUtilities::ConvertMinFilter(TextureFilter);
-	const EGLTFJsonTextureFilter TextureMagFilter = FGLTFCoreUtilities::ConvertMagFilter(TextureFilter);
 
 	FGLTFJsonTexture* Texture = FGLTFMaterialUtility::AddTexture(
 		Builder,
@@ -1507,10 +1476,8 @@ bool FGLTFDelayedMaterialTask::StoreBakedPropertyTexture(FGLTFJsonTextureInfo& O
 		true, // NOTE: we can ignore alpha in everything but TryGetBaseColorAndOpacity
 		FGLTFMaterialUtility::IsNormalMap(PropertyBakeOutput.Property),
 		GetBakedTextureName(PropertyName),
-		TextureMinFilter,
-		TextureMagFilter,
-		TextureWrapS,
-		TextureWrapT);
+		TextureAddress,
+		TextureFilter);
 
 	OutTexInfo.Index = Texture;
 	return true;
diff --git a/Engine/Plugins/Enterprise/GLTFExporter/Source/GLTFExporter/Private/Utilities/GLTFCoreUtilities.cpp b/Engine/Plugins/Enterprise/GLTFExporter/Source/GLTFExporter/Private/Utilities/GLTFCoreUtilities.cpp
--- a/Engine/Plugins/Enterprise/GLTFExporter/Source/GLTFExporter/Private/Utilities/GLTFCoreUtilities.cpp
+++ b/Engine/Plugins/Enterprise/GLTFExporter/Source/GLTFExporter/Private/Utilities/GLTFCoreUtilities.cpp
@@ -237,7 +237,7 @@ EGLTFJsonTextureWrap FGLTFCoreUtilities::ConvertWrap(TextureAddress Address)
 		case TA_Wrap:   return EGLTFJsonTextureWrap::Repeat;
 		case TA_Mirror: return EGLTFJsonTextureWrap::MirroredRepeat;
 		case TA_Clamp:  return EGLTFJsonTextureWrap::ClampToEdge;
-		default:        return EGLTFJsonTextureWrap::None; // TODO: add error handling in callers
+		default:        return EGLTFJsonTextureWrap::None;
 	}
 }
 
@@ -262,13 +262,3 @@ EGLTFJsonTextureFilter FGLTFCoreUtilities::ConvertMagFilter(TextureFilter Filter
 		default:           return EGLTFJsonTextureFilter::None;
 	}
 }
-
-EGLTFJsonTextureFilter FGLTFCoreUtilities::ConvertMinFilter(TextureFilter Filter, TextureGroup LODGroup)
-{
-	return ConvertMinFilter(Filter == TF_Default ? FGLTFTextureUtility::GetDefaultFilter(LODGroup) : Filter);
-}
-
-EGLTFJsonTextureFilter FGLTFCoreUtilities::ConvertMagFilter(TextureFilter Filter, TextureGroup LODGroup)
-{
-	return ConvertMagFilter(Filter == TF_Default ? FGLTFTextureUtility::GetDefaultFilter(LODGroup) : Filter);
-}
diff --git a/Engine/Plugins/Enterprise/GLTFExporter/Source/GLTFExporter/Public/Builders/GLTFConvertBuilder.h b/Engine/Plugins/Enterprise/GLTFExporter/Source/GLTFExporter/Public/Builders/GLTFConvertBuilder.h
--- a/Engine/Plugins/Enterprise/GLTFExporter/Source/GLTFExporter/Public/Builders/GLTFConvertBuilder.h
+++ b/Engine/Plugins/Enterprise/GLTFExporter/Source/GLTFExporter/Public/Builders/GLTFConvertBuilder.h
@@ -58,6 +58,9 @@ public:
 	FGLTFJsonMaterial* AddUniqueMaterial(const UMaterialInterface* Material, const FGLTFMeshData* MeshData = nullptr, const FGLTFIndexArray& SectionIndices = {});
 
 	FGLTFJsonSampler* AddUniqueSampler(const UTexture* Texture);
+	FGLTFJsonSampler* AddUniqueSampler(TextureAddress Address, TextureFilter Filter, TextureGroup LODGroup = TEXTUREGROUP_MAX);
+	FGLTFJsonSampler* AddUniqueSampler(TextureAddress AddressX, TextureAddress AddressY, TextureFilter Filter, TextureGroup LODGroup = TEXTUREGROUP_MAX);
+
 	FGLTFJsonTexture* AddUniqueTexture(const UTexture* Texture);
 	FGLTFJsonTexture* AddUniqueTexture(const UTexture2D* Texture);
 	FGLTFJsonTexture* AddUniqueTexture(const UTextureRenderTarget2D* Texture);
@@ -105,7 +108,9 @@ public:
 	TUniquePtr<IGLTFStaticMeshDataConverter> StaticMeshDataConverter = MakeUnique<FGLTFStaticMeshDataConverter>(*this);
 	TUniquePtr<IGLTFSkeletalMeshDataConverter> SkeletalMeshDataConverter = MakeUnique<FGLTFSkeletalMeshDataConverter>(*this);
 
+	TUniquePtr<IGLTFTextureSamplerConverter> TextureSamplerConverter = MakeUnique<FGLTFTextureSamplerConverter>(*this);
 	TUniquePtr<IGLTFSamplerConverter> SamplerConverter = MakeUnique<FGLTFSamplerConverter>(*this);
+
 	TUniquePtr<IGLTFTexture2DConverter> Texture2DConverter = MakeUnique<FGLTFTexture2DConverter>(*this);
 	TUniquePtr<IGLTFTextureRenderTarget2DConverter> TextureRenderTarget2DConverter = MakeUnique<FGLTFTextureRenderTarget2DConverter>(*this);
 	TUniquePtr<IGLTFImageConverter> ImageConverter = MakeUnique<FGLTFImageConverter>(*this);
diff --git a/Engine/Plugins/Enterprise/GLTFExporter/Source/GLTFExporter/Public/Converters/GLTFSamplerConverters.h b/Engine/Plugins/Enterprise/GLTFExporter/Source/GLTFExporter/Public/Converters/GLTFSamplerConverters.h
--- a/Engine/Plugins/Enterprise/GLTFExporter/Source/GLTFExporter/Public/Converters/GLTFSamplerConverters.h
+++ b/Engine/Plugins/Enterprise/GLTFExporter/Source/GLTFExporter/Public/Converters/GLTFSamplerConverters.h
@@ -5,12 +5,14 @@
 #include "Json/GLTFJsonCore.h"
 #include "Converters/GLTFConverter.h"
 #include "Converters/GLTFBuilderContext.h"
+#include "Engine/TextureDefines.h"
 
 class UTexture;
 
-typedef TGLTFConverter<FGLTFJsonSampler*, const UTexture*> IGLTFSamplerConverter;
+typedef TGLTFConverter<FGLTFJsonSampler*, const UTexture*> IGLTFTextureSamplerConverter;
+typedef TGLTFConverter<FGLTFJsonSampler*, TextureAddress, TextureAddress, TextureFilter, TextureGroup> IGLTFSamplerConverter;
 
-class GLTFEXPORTER_API FGLTFSamplerConverter : public FGLTFBuilderContext, public IGLTFSamplerConverter
+class GLTFEXPORTER_API FGLTFTextureSamplerConverter : public FGLTFBuilderContext, public IGLTFTextureSamplerConverter
 {
 public:
 
@@ -20,3 +22,16 @@ protected:
 
 	virtual FGLTFJsonSampler* Convert(const UTexture* Texture) override;
 };
+
+class GLTFEXPORTER_API FGLTFSamplerConverter : public FGLTFBuilderContext, public IGLTFSamplerConverter
+{
+public:
+
+	using FGLTFBuilderContext::FGLTFBuilderContext;
+
+protected:
+
+	virtual void Sanitize(TextureAddress& AddressX, TextureAddress& AddressY, TextureFilter& Filter, TextureGroup& LODGroup) override;
+
+	virtual FGLTFJsonSampler* Convert(TextureAddress AddressX, TextureAddress AddressY, TextureFilter Filter, TextureGroup LODGroup) override;
+};
diff --git a/Engine/Plugins/Enterprise/GLTFExporter/Source/GLTFExporter/Public/Utilities/GLTFCoreUtilities.h b/Engine/Plugins/Enterprise/GLTFExporter/Source/GLTFExporter/Public/Utilities/GLTFCoreUtilities.h
--- a/Engine/Plugins/Enterprise/GLTFExporter/Source/GLTFExporter/Public/Utilities/GLTFCoreUtilities.h
+++ b/Engine/Plugins/Enterprise/GLTFExporter/Source/GLTFExporter/Public/Utilities/GLTFCoreUtilities.h
@@ -68,9 +68,6 @@ struct GLTFEXPORTER_API FGLTFCoreUtilities
 	static EGLTFJsonTextureFilter ConvertMinFilter(TextureFilter Filter);
 	static EGLTFJsonTextureFilter ConvertMagFilter(TextureFilter Filter);
 
-	static EGLTFJsonTextureFilter ConvertMinFilter(TextureFilter Filter, TextureGroup LODGroup);
-	static EGLTFJsonTextureFilter ConvertMagFilter(TextureFilter Filter, TextureGroup LODGroup);
-
 	template <typename ComponentType>
 	static EGLTFJsonComponentType GetComponentType()
 	{
