Fix GLTFMaterialExporter depending on EditorSphere in runtime, as preview mesh when exporting standalone materials (which otherwise don't have any mesh on their own). Instead try to load EditorSphere on demand and if not available try other sphere meshes in engine content.

diff --git a/Engine/Plugins/Enterprise/GLTFExporter/Source/GLTFExporter/Private/Exporters/GLTFExporterUtility.cpp b/Engine/Plugins/Enterprise/GLTFExporter/Source/GLTFExporter/Private/Exporters/GLTFExporterUtility.cpp
--- a/Engine/Plugins/Enterprise/GLTFExporter/Source/GLTFExporter/Private/Exporters/GLTFExporterUtility.cpp
+++ b/Engine/Plugins/Enterprise/GLTFExporter/Source/GLTFExporter/Private/Exporters/GLTFExporterUtility.cpp
@@ -43,7 +43,23 @@ const UStaticMesh* FGLTFExporterUtility::GetPreviewMesh(const UMaterialInterface
 	} while (Material != nullptr);
 #endif
 
-	return nullptr;
+	static const UStaticMesh* DefaultPreviewMesh = GetSphereMesh();
+	return DefaultPreviewMesh;
+}
+
+const UStaticMesh* FGLTFExporterUtility::GetSphereMesh()
+{
+	const UStaticMesh* SphereMesh = LoadObject<UStaticMesh>(nullptr, TEXT("/Engine/EditorMeshes/EditorSphere.EditorSphere"));
+	if (SphereMesh == nullptr)
+	{
+		SphereMesh = LoadObject<UStaticMesh>(nullptr, TEXT("/Engine/BasicShapes/Sphere.Sphere"));
+		if (SphereMesh == nullptr)
+		{
+			SphereMesh = LoadObject<UStaticMesh>(nullptr, TEXT("/Engine/EngineMeshes/Sphere.Sphere"));
+		}
+	}
+
+	return SphereMesh;
 }
 
 const USkeletalMesh* FGLTFExporterUtility::GetPreviewMesh(const UAnimSequence* AnimSequence)
diff --git a/Engine/Plugins/Enterprise/GLTFExporter/Source/GLTFExporter/Private/Exporters/GLTFExporterUtility.h b/Engine/Plugins/Enterprise/GLTFExporter/Source/GLTFExporter/Private/Exporters/GLTFExporterUtility.h
--- a/Engine/Plugins/Enterprise/GLTFExporter/Source/GLTFExporter/Private/Exporters/GLTFExporterUtility.h
+++ b/Engine/Plugins/Enterprise/GLTFExporter/Source/GLTFExporter/Private/Exporters/GLTFExporterUtility.h
@@ -16,6 +16,8 @@ struct FGLTFExporterUtility
 	static void GetSelectedActors(TSet<AActor*>& OutSelectedActors);
 
 	static const UStaticMesh* GetPreviewMesh(const UMaterialInterface* Material);
+	static const UStaticMesh* GetSphereMesh();
+
 	static const USkeletalMesh* GetPreviewMesh(const UAnimSequence* AnimSequence);
 	static const USkeletalMesh* FindCompatibleMesh(const USkeleton* Skeleton);
 
diff --git a/Engine/Plugins/Enterprise/GLTFExporter/Source/GLTFExporter/Private/Exporters/GLTFMaterialExporter.cpp b/Engine/Plugins/Enterprise/GLTFExporter/Source/GLTFExporter/Private/Exporters/GLTFMaterialExporter.cpp
--- a/Engine/Plugins/Enterprise/GLTFExporter/Source/GLTFExporter/Private/Exporters/GLTFMaterialExporter.cpp
+++ b/Engine/Plugins/Enterprise/GLTFExporter/Source/GLTFExporter/Private/Exporters/GLTFMaterialExporter.cpp
@@ -12,8 +12,7 @@ UGLTFMaterialExporter::UGLTFMaterialExporter(const FObjectInitializer& ObjectIni
 {
 	SupportedClass = UMaterialInterface::StaticClass();
 
-	static ConstructorHelpers::FObjectFinder<UStaticMesh> DefaultPreviewMeshFinder(TEXT("/Engine/EditorMeshes/EditorSphere.EditorSphere"));
-	DefaultPreviewMesh = DefaultPreviewMeshFinder.Object;
+	 // TODO: remove UGLTFMaterialExporter::DefaultPreviewMesh in 5.2 since it's no longer in use.
 }
 
 bool UGLTFMaterialExporter::AddObject(FGLTFContainerBuilder& Builder, const UObject* Object)
@@ -23,28 +22,33 @@ bool UGLTFMaterialExporter::AddObject(FGLTFContainerBuilder& Builder, const UObj
 	if (Builder.ExportOptions->bExportPreviewMesh)
 	{
 		const UStaticMesh* PreviewMesh = FGLTFExporterUtility::GetPreviewMesh(Material);
-		if (PreviewMesh == nullptr)
+		if (PreviewMesh != nullptr)
 		{
-			PreviewMesh = DefaultPreviewMesh;
+			FGLTFJsonMesh* Mesh = Builder.AddUniqueMesh(PreviewMesh, { Material });
+			if (Mesh == nullptr)
+			{
+				Builder.LogError(
+					FString::Printf(TEXT("Failed to export preview mesh %s for material %s"),
+					*Material->GetName(),
+					*PreviewMesh->GetName()));
+				return false;
+			}
+
+			FGLTFJsonNode* Node = Builder.AddNode();
+			Node->Mesh = Mesh;
+
+			FGLTFJsonScene* Scene = Builder.AddScene();
+			Scene->Nodes.Add(Node);
+
+			Builder.DefaultScene = Scene;
 		}
-
-		FGLTFJsonMesh* Mesh = Builder.AddUniqueMesh(PreviewMesh, { Material });
-		if (Mesh == nullptr)
+		else
 		{
 			Builder.LogError(
-				FString::Printf(TEXT("Failed to export preview mesh %s for material %s"),
-				*Material->GetName(),
-				*PreviewMesh->GetName()));
+				FString::Printf(TEXT("Failed to export material %s because of missing preview mesh"),
+				*Material->GetName()));
 			return false;
 		}
-
-		FGLTFJsonNode* Node = Builder.AddNode();
-		Node->Mesh = Mesh;
-
-		FGLTFJsonScene* Scene = Builder.AddScene();
-		Scene->Nodes.Add(Node);
-
-		Builder.DefaultScene = Scene;
 	}
 	else
 	{
