Added missing include to fix non unity builds

diff --git a/Engine/Plugins/Enterprise/GLTFExporter/Source/GLTFExporter/Private/Exporters/GLTFExporterUtility.cpp b/Engine/Plugins/Enterprise/GLTFExporter/Source/GLTFExporter/Private/Exporters/GLTFExporterUtility.cpp
--- a/Engine/Plugins/Enterprise/GLTFExporter/Source/GLTFExporter/Private/Exporters/GLTFExporterUtility.cpp
+++ b/Engine/Plugins/Enterprise/GLTFExporter/Source/GLTFExporter/Private/Exporters/GLTFExporterUtility.cpp
@@ -7,6 +7,7 @@
 #include "Animation/AnimSequence.h"
 #include "Animation/Skeleton.h"
 #include "Engine/SkeletalMesh.h"
+#include "Engine/StaticMesh.h"
 #include "Engine/Selection.h"
 #include "Engine/World.h"
 #if WITH_EDITOR
