Simplify sampler conversion to reduce unnecessary complexity in gltf exporter

diff --git a/Engine/Plugins/Enterprise/GLTFExporter/Source/GLTFExporter/Private/Builders/GLTFConvertBuilder.cpp b/Engine/Plugins/Enterprise/GLTFExporter/Source/GLTFExporter/Private/Builders/GLTFConvertBuilder.cpp
--- a/Engine/Plugins/Enterprise/GLTFExporter/Source/GLTFExporter/Private/Builders/GLTFConvertBuilder.cpp
+++ b/Engine/Plugins/Enterprise/GLTFExporter/Source/GLTFExporter/Private/Builders/GLTFConvertBuilder.cpp
@@ -247,7 +247,7 @@ FGLTFJsonSampler* FGLTFConvertBuilder::AddUniqueSampler(const UTexture* Texture)
 		return nullptr;
 	}
 
-	return TextureSamplerConverter->GetOrAdd(Texture);
+	return AddUniqueSampler(Texture->GetTextureAddressX(), Texture->GetTextureAddressY(), Texture->Filter, Texture->LODGroup);
 }
 
 FGLTFJsonSampler* FGLTFConvertBuilder::AddUniqueSampler(TextureAddress Address, TextureFilter Filter, TextureGroup LODGroup)
diff --git a/Engine/Plugins/Enterprise/GLTFExporter/Source/GLTFExporter/Private/Converters/GLTFSamplerConverters.cpp b/Engine/Plugins/Enterprise/GLTFExporter/Source/GLTFExporter/Private/Converters/GLTFSamplerConverters.cpp
--- a/Engine/Plugins/Enterprise/GLTFExporter/Source/GLTFExporter/Private/Converters/GLTFSamplerConverters.cpp
+++ b/Engine/Plugins/Enterprise/GLTFExporter/Source/GLTFExporter/Private/Converters/GLTFSamplerConverters.cpp
@@ -5,24 +5,6 @@
 #include "Utilities/GLTFCoreUtilities.h"
 #include "Converters/GLTFTextureUtility.h"
 
-FGLTFJsonSampler* FGLTFTextureSamplerConverter::Convert(const UTexture* Texture)
-{
-	TextureAddress AddressX;
-	TextureAddress AddressY;
-
-	if (FGLTFTextureUtility::IsCubemap(Texture))
-	{
-		AddressX = TA_Clamp;
-		AddressY = TA_Clamp;
-	}
-	else
-	{
-		FGLTFTextureUtility::GetAddressXY(Texture, AddressX, AddressY);
-	}
-
-	return Builder.AddUniqueSampler(AddressX, AddressY, Texture->Filter, Texture->LODGroup);
-}
-
 void FGLTFSamplerConverter::Sanitize(TextureAddress& AddressX, TextureAddress& AddressY, TextureFilter& Filter, TextureGroup& LODGroup)
 {
 	if (Filter == TF_Default)
diff --git a/Engine/Plugins/Enterprise/GLTFExporter/Source/GLTFExporter/Private/Converters/GLTFTextureUtility.cpp b/Engine/Plugins/Enterprise/GLTFExporter/Source/GLTFExporter/Private/Converters/GLTFTextureUtility.cpp
--- a/Engine/Plugins/Enterprise/GLTFExporter/Source/GLTFExporter/Private/Converters/GLTFTextureUtility.cpp
+++ b/Engine/Plugins/Enterprise/GLTFExporter/Source/GLTFExporter/Private/Converters/GLTFTextureUtility.cpp
@@ -103,25 +103,6 @@ FIntPoint FGLTFTextureUtility::GetInGameSize(const UTexture* Texture)
 	return { InGameWidth, InGameHeight };
 }
 
-void FGLTFTextureUtility::GetAddressXY(const UTexture* Texture, TextureAddress& OutAddressX, TextureAddress& OutAddressY)
-{
-	if (const UTexture2D* Texture2D = Cast<UTexture2D>(Texture))
-	{
-		OutAddressX = Texture2D->AddressX;
-		OutAddressY = Texture2D->AddressY;
-	}
-	else if (const UTextureRenderTarget2D* RenderTarget2D = Cast<UTextureRenderTarget2D>(Texture))
-	{
-		OutAddressX = RenderTarget2D->AddressX;
-		OutAddressY = RenderTarget2D->AddressY;
-	}
-	else
-	{
-		OutAddressX = TA_MAX;
-		OutAddressY = TA_MAX;
-	}
-}
-
 UTextureRenderTarget2D* FGLTFTextureUtility::CreateRenderTarget(const FIntPoint& Size, bool bHDR)
 {
 	// TODO: instead of PF_FloatRGBA (i.e. RTF_RGBA16f) use PF_A32B32G32R32F (i.e. RTF_RGBA32f) to avoid accuracy loss
diff --git a/Engine/Plugins/Enterprise/GLTFExporter/Source/GLTFExporter/Private/Converters/GLTFTextureUtility.h b/Engine/Plugins/Enterprise/GLTFExporter/Source/GLTFExporter/Private/Converters/GLTFTextureUtility.h
--- a/Engine/Plugins/Enterprise/GLTFExporter/Source/GLTFExporter/Private/Converters/GLTFTextureUtility.h
+++ b/Engine/Plugins/Enterprise/GLTFExporter/Source/GLTFExporter/Private/Converters/GLTFTextureUtility.h
@@ -24,8 +24,6 @@ struct FGLTFTextureUtility
 
 	static FIntPoint GetInGameSize(const UTexture* Texture);
 
-	static void GetAddressXY(const UTexture* Texture, TextureAddress& OutAddressX, TextureAddress& OutAddressY);
-
 	static UTextureRenderTarget2D* CreateRenderTarget(const FIntPoint& Size, bool bIsHDR);
 
 	static bool DrawTexture(UTextureRenderTarget2D* OutTarget, const UTexture2D* InSource);
diff --git a/Engine/Plugins/Enterprise/GLTFExporter/Source/GLTFExporter/Public/Builders/GLTFConvertBuilder.h b/Engine/Plugins/Enterprise/GLTFExporter/Source/GLTFExporter/Public/Builders/GLTFConvertBuilder.h
--- a/Engine/Plugins/Enterprise/GLTFExporter/Source/GLTFExporter/Public/Builders/GLTFConvertBuilder.h
+++ b/Engine/Plugins/Enterprise/GLTFExporter/Source/GLTFExporter/Public/Builders/GLTFConvertBuilder.h
@@ -108,7 +108,6 @@ public:
 	TUniquePtr<IGLTFStaticMeshDataConverter> StaticMeshDataConverter = MakeUnique<FGLTFStaticMeshDataConverter>(*this);
 	TUniquePtr<IGLTFSkeletalMeshDataConverter> SkeletalMeshDataConverter = MakeUnique<FGLTFSkeletalMeshDataConverter>(*this);
 
-	TUniquePtr<IGLTFTextureSamplerConverter> TextureSamplerConverter = MakeUnique<FGLTFTextureSamplerConverter>(*this);
 	TUniquePtr<IGLTFSamplerConverter> SamplerConverter = MakeUnique<FGLTFSamplerConverter>(*this);
 
 	TUniquePtr<IGLTFTexture2DConverter> Texture2DConverter = MakeUnique<FGLTFTexture2DConverter>(*this);
diff --git a/Engine/Plugins/Enterprise/GLTFExporter/Source/GLTFExporter/Public/Converters/GLTFSamplerConverters.h b/Engine/Plugins/Enterprise/GLTFExporter/Source/GLTFExporter/Public/Converters/GLTFSamplerConverters.h
--- a/Engine/Plugins/Enterprise/GLTFExporter/Source/GLTFExporter/Public/Converters/GLTFSamplerConverters.h
+++ b/Engine/Plugins/Enterprise/GLTFExporter/Source/GLTFExporter/Public/Converters/GLTFSamplerConverters.h
@@ -7,22 +7,8 @@
 #include "Converters/GLTFBuilderContext.h"
 #include "Engine/TextureDefines.h"
 
-class UTexture;
-
-typedef TGLTFConverter<FGLTFJsonSampler*, const UTexture*> IGLTFTextureSamplerConverter;
 typedef TGLTFConverter<FGLTFJsonSampler*, TextureAddress, TextureAddress, TextureFilter, TextureGroup> IGLTFSamplerConverter;
 
-class GLTFEXPORTER_API FGLTFTextureSamplerConverter : public FGLTFBuilderContext, public IGLTFTextureSamplerConverter
-{
-public:
-
-	using FGLTFBuilderContext::FGLTFBuilderContext;
-
-protected:
-
-	virtual FGLTFJsonSampler* Convert(const UTexture* Texture) override;
-};
-
 class GLTFEXPORTER_API FGLTFSamplerConverter : public FGLTFBuilderContext, public IGLTFSamplerConverter
 {
 public:
