Add ClearCoatBottomNormal and TransmittanceColor in FGLTFMaterialPropertyEx

diff --git a/Engine/Plugins/Enterprise/GLTFExporter/Source/GLTFExporter/Private/Converters/GLTFMaterialUtility.cpp b/Engine/Plugins/Enterprise/GLTFExporter/Source/GLTFExporter/Private/Converters/GLTFMaterialUtility.cpp
--- a/Engine/Plugins/Enterprise/GLTFExporter/Source/GLTFExporter/Private/Converters/GLTFMaterialUtility.cpp
+++ b/Engine/Plugins/Enterprise/GLTFExporter/Source/GLTFExporter/Private/Converters/GLTFMaterialUtility.cpp
@@ -27,12 +27,12 @@ UMaterialInterface* FGLTFMaterialUtility::GetDefaultMaterial()
 
 bool FGLTFMaterialUtility::IsNormalMap(const FGLTFMaterialPropertyEx& Property)
 {
-	return Property == MP_Normal || Property == TEXT("ClearCoatBottomNormal");
+	return Property == MP_Normal || Property == FGLTFMaterialPropertyEx::ClearCoatBottomNormal;
 }
 
 bool FGLTFMaterialUtility::IsSRGB(const FGLTFMaterialPropertyEx& Property)
 {
-	return Property == MP_BaseColor || Property == MP_EmissiveColor || Property == MP_SubsurfaceColor || Property == TEXT("TransmittanceColor");
+	return Property == MP_BaseColor || Property == MP_EmissiveColor || Property == MP_SubsurfaceColor || Property == FGLTFMaterialPropertyEx::TransmittanceColor;
 }
 
 FGuid FGLTFMaterialUtility::GetAttributeID(const FGLTFMaterialPropertyEx& Property)
@@ -315,7 +315,10 @@ bool FGLTFMaterialUtility::NeedsMeshData(const UMaterialInterface* Material)
 			MP_AmbientOcclusion,
 			MP_CustomData0,
 			MP_CustomData1,
-			TEXT("ClearCoatBottomNormal"),
+			FGLTFMaterialPropertyEx::ClearCoatBottomNormal,
+			// TODO: add TransmittanceColor when supported
+			// TODO: add Refraction when supported
+			// TODO: add Specular when supported
 		};
 
 		bool bNeedsMeshData = false;
diff --git a/Engine/Plugins/Enterprise/GLTFExporter/Source/GLTFExporter/Private/Tasks/GLTFDelayedMaterialTasks.cpp b/Engine/Plugins/Enterprise/GLTFExporter/Source/GLTFExporter/Private/Tasks/GLTFDelayedMaterialTasks.cpp
--- a/Engine/Plugins/Enterprise/GLTFExporter/Source/GLTFExporter/Private/Tasks/GLTFDelayedMaterialTasks.cpp
+++ b/Engine/Plugins/Enterprise/GLTFExporter/Source/GLTFExporter/Private/Tasks/GLTFDelayedMaterialTasks.cpp
@@ -115,7 +115,7 @@ void FGLTFDelayedMaterialTask::Process()
 				Builder.LogWarning(FString::Printf(TEXT("Failed to export %s for material %s"), *EmissiveProperty.ToString(), *Material->GetName()));
 			}
 
-			const FGLTFMaterialPropertyEx NormalProperty = JsonMaterial->ShadingModel == EGLTFJsonShadingModel::ClearCoat ? FGLTFMaterialPropertyEx(TEXT("ClearCoatBottomNormal")) : FGLTFMaterialPropertyEx(MP_Normal);
+			const FGLTFMaterialPropertyEx NormalProperty = JsonMaterial->ShadingModel == EGLTFJsonShadingModel::ClearCoat ? FGLTFMaterialPropertyEx::ClearCoatBottomNormal : FGLTFMaterialPropertyEx(MP_Normal);
 			if (IsPropertyNonDefault(NormalProperty))
 			{
 				if (!TryGetSourceTexture(JsonMaterial->NormalTexture, NormalProperty, DefaultColorInputMasks))
@@ -1504,7 +1504,7 @@ EGLTFMaterialPropertyGroup FGLTFDelayedMaterialTask::GetPropertyGroup(const FGLT
 		case MP_CustomData1:
 			return EGLTFMaterialPropertyGroup::ClearCoatRoughness;
 		case MP_CustomOutput:
-			if (Property.CustomOutput == TEXT("ClearCoatBottomNormal"))
+			if (Property == FGLTFMaterialPropertyEx::ClearCoatBottomNormal)
 			{
 				return EGLTFMaterialPropertyGroup::ClearCoatBottomNormal;
 			}
diff --git a/Engine/Plugins/Enterprise/GLTFExporter/Source/GLTFMaterialBaking/Private/GLTFMaterialBakingModule.cpp b/Engine/Plugins/Enterprise/GLTFExporter/Source/GLTFMaterialBaking/Private/GLTFMaterialBakingModule.cpp
--- a/Engine/Plugins/Enterprise/GLTFExporter/Source/GLTFMaterialBaking/Private/GLTFMaterialBakingModule.cpp
+++ b/Engine/Plugins/Enterprise/GLTFExporter/Source/GLTFMaterialBaking/Private/GLTFMaterialBakingModule.cpp
@@ -288,8 +288,8 @@ void FGLTFMaterialBakingModule::StartupModule()
 	PerPropertyFormat.Add(MP_CustomData1, PF_B8G8R8A8);
 	PerPropertyFormat.Add(MP_Refraction, PF_B8G8R8A8);
 	PerPropertyFormat.Add(MP_ShadingModel, PF_B8G8R8A8);
-	PerPropertyFormat.Add(TEXT("ClearCoatBottomNormal"), PF_B8G8R8A8);
-	PerPropertyFormat.Add(TEXT("TransmittanceColor"), PF_B8G8R8A8);
+	PerPropertyFormat.Add(FGLTFMaterialPropertyEx::ClearCoatBottomNormal, PF_B8G8R8A8);
+	PerPropertyFormat.Add(FGLTFMaterialPropertyEx::TransmittanceColor, PF_B8G8R8A8);
 
 	// Register callback for modified objects
 	FCoreUObjectDelegates::OnObjectModified.AddRaw(this, &FGLTFMaterialBakingModule::OnObjectModified);
@@ -789,7 +789,7 @@ void FGLTFMaterialBakingModule::SetLinearBake(bool bCorrectLinear)
 		PerPropertyColorSpace.Add(MP_BaseColor, EPropertyColorSpace::sRGB);
 		PerPropertyColorSpace.Add(MP_EmissiveColor, EPropertyColorSpace::sRGB);
 		PerPropertyColorSpace.Add(MP_SubsurfaceColor, EPropertyColorSpace::sRGB);
-		PerPropertyColorSpace.Add(TEXT("TransmittanceColor"), EPropertyColorSpace::sRGB);
+		PerPropertyColorSpace.Add(FGLTFMaterialPropertyEx::TransmittanceColor, EPropertyColorSpace::sRGB);
 	}
 	else
 	{
@@ -799,7 +799,7 @@ void FGLTFMaterialBakingModule::SetLinearBake(bool bCorrectLinear)
 		PerPropertyColorSpace.Add(MP_OpacityMask, EPropertyColorSpace::Linear);
 		PerPropertyColorSpace.Add(MP_Refraction, EPropertyColorSpace::Linear);
 		PerPropertyColorSpace.Add(MP_ShadingModel, EPropertyColorSpace::Linear);
-		PerPropertyColorSpace.Add(TEXT("ClearCoatBottomNormal"), EPropertyColorSpace::Linear);
+		PerPropertyColorSpace.Add(FGLTFMaterialPropertyEx::ClearCoatBottomNormal, EPropertyColorSpace::Linear);
 	}
 }
 
diff --git a/Engine/Plugins/Enterprise/GLTFExporter/Source/GLTFMaterialBaking/Private/GLTFMaterialPropertyEx.cpp b/Engine/Plugins/Enterprise/GLTFExporter/Source/GLTFMaterialBaking/Private/GLTFMaterialPropertyEx.cpp
--- a/Engine/Plugins/Enterprise/GLTFExporter/Source/GLTFMaterialBaking/Private/GLTFMaterialPropertyEx.cpp
+++ b/Engine/Plugins/Enterprise/GLTFExporter/Source/GLTFMaterialBaking/Private/GLTFMaterialPropertyEx.cpp
@@ -3,6 +3,9 @@
 #include "GLTFMaterialPropertyEx.h"
 #include "UObject/NameTypes.h"
 
+const FGLTFMaterialPropertyEx FGLTFMaterialPropertyEx::ClearCoatBottomNormal(TEXT("ClearCoatBottomNormal"));
+const FGLTFMaterialPropertyEx FGLTFMaterialPropertyEx::TransmittanceColor(TEXT("TransmittanceColor"));
+
 FString FGLTFMaterialPropertyEx::ToString() const
 {
 	if (!IsCustomOutput())
diff --git a/Engine/Plugins/Enterprise/GLTFExporter/Source/GLTFMaterialBaking/Public/GLTFMaterialPropertyEx.h b/Engine/Plugins/Enterprise/GLTFExporter/Source/GLTFMaterialBaking/Public/GLTFMaterialPropertyEx.h
--- a/Engine/Plugins/Enterprise/GLTFExporter/Source/GLTFMaterialBaking/Public/GLTFMaterialPropertyEx.h
+++ b/Engine/Plugins/Enterprise/GLTFExporter/Source/GLTFMaterialBaking/Public/GLTFMaterialPropertyEx.h
@@ -5,7 +5,7 @@
 #include "SceneTypes.h"
 
 /** Structure extending EMaterialProperty to allow detailed information about custom output */
-struct FGLTFMaterialPropertyEx
+struct GLTFMATERIALBAKING_API FGLTFMaterialPropertyEx
 {
 	FGLTFMaterialPropertyEx(EMaterialProperty Type = MP_MAX, const FName& CustomOutput = NAME_None)
 		: Type(Type)
@@ -42,11 +42,14 @@ struct FGLTFMaterialPropertyEx
 		return !Other.IsCustomOutput() ? GetTypeHash(Other.Type) : GetTypeHash(Other.CustomOutput);
 	}
 
-	GLTFMATERIALBAKING_API FString ToString() const;
+	FString ToString() const;
 
 	/** The material property */
 	EMaterialProperty Type;
 
 	/** The name of a specific custom output. Only used if property is MP_CustomOutput */
 	FName CustomOutput;
+
+	static const FGLTFMaterialPropertyEx ClearCoatBottomNormal;
+	static const FGLTFMaterialPropertyEx TransmittanceColor;
 };
