Warn instead of attempting to export ClearCoatBottomNormal if disabled in render settings

diff --git a/Engine/Plugins/Enterprise/GLTFExporter/Source/GLTFExporter/Private/Converters/GLTFMaterialUtility.cpp b/Engine/Plugins/Enterprise/GLTFExporter/Source/GLTFExporter/Private/Converters/GLTFMaterialUtility.cpp
--- a/Engine/Plugins/Enterprise/GLTFExporter/Source/GLTFExporter/Private/Converters/GLTFMaterialUtility.cpp
+++ b/Engine/Plugins/Enterprise/GLTFExporter/Source/GLTFExporter/Private/Converters/GLTFMaterialUtility.cpp
@@ -11,10 +11,10 @@
 #include "IGLTFMaterialBakingModule.h"
 #include "GLTFMaterialBakingStructures.h"
 #endif
+#include "Engine/RendererSettings.h"
 #include "Modules/ModuleManager.h"
 #include "Materials/MaterialInstanceConstant.h"
 #include "Materials/MaterialExpressionCustomOutput.h"
-#include "Materials/MaterialExpressionClearCoatNormalCustomOutput.h"
 #include "Materials/MaterialExpressionTextureCoordinate.h"
 
 UMaterialInterface* FGLTFMaterialUtility::GetDefaultMaterial()
@@ -23,6 +23,11 @@ UMaterialInterface* FGLTFMaterialUtility::GetDefaultMaterial()
 	return DefaultMaterial;
 }
 
+bool FGLTFMaterialUtility::IsClearCoatBottomNormalEnabled()
+{
+	return GetDefault<URendererSettings>()->bClearCoatEnableSecondNormal != 0;
+}
+
 #if WITH_EDITOR
 
 bool FGLTFMaterialUtility::IsNormalMap(const FGLTFMaterialPropertyEx& Property)
@@ -73,22 +78,27 @@ const FExpressionInput* FGLTFMaterialUtility::GetInputForProperty(const UMateria
 {
 	if (Property.IsCustomOutput())
 	{
-		const UMaterialExpressionCustomOutput* CustomOutput = GetCustomOutputByName(Material, Property.CustomOutput.ToString());
-		return CustomOutput != nullptr ? &CastChecked<UMaterialExpressionClearCoatNormalCustomOutput>(CustomOutput)->Input : nullptr;
+		const FString FunctionName = Property.CustomOutput.ToString();
+		const UMaterialExpressionCustomOutput* CustomOutput = GetCustomOutputByName(Material, FunctionName);
+		if (CustomOutput == nullptr)
+		{
+			return nullptr;
+		}
+
+		// Assume custom outputs always have a single input (which is true for all supported custom outputs)
+		return const_cast<UMaterialExpressionCustomOutput*>(CustomOutput)->GetInput(0);
 	}
 
 	UMaterial* UnderlyingMaterial = const_cast<UMaterial*>(Material->GetMaterial());
 	return UnderlyingMaterial->GetExpressionInputForProperty(Property.Type);
 }
 
-const UMaterialExpressionCustomOutput* FGLTFMaterialUtility::GetCustomOutputByName(const UMaterialInterface* Material, const FString& Name)
+const UMaterialExpressionCustomOutput* FGLTFMaterialUtility::GetCustomOutputByName(const UMaterialInterface* Material, const FString& FunctionName)
 {
-	// TODO: should we also search inside material functions and attribute layers?
-
 	for (const TObjectPtr<UMaterialExpression>& Expression : Material->GetMaterial()->GetExpressions())
 	{
 		const UMaterialExpressionCustomOutput* CustomOutput = Cast<UMaterialExpressionCustomOutput>(Expression);
-		if (CustomOutput != nullptr && CustomOutput->GetDisplayName() == Name)
+		if (CustomOutput != nullptr && CustomOutput->GetFunctionName() == FunctionName)
 		{
 			return CustomOutput;
 		}
diff --git a/Engine/Plugins/Enterprise/GLTFExporter/Source/GLTFExporter/Private/Converters/GLTFMaterialUtility.h b/Engine/Plugins/Enterprise/GLTFExporter/Source/GLTFExporter/Private/Converters/GLTFMaterialUtility.h
--- a/Engine/Plugins/Enterprise/GLTFExporter/Source/GLTFExporter/Private/Converters/GLTFMaterialUtility.h
+++ b/Engine/Plugins/Enterprise/GLTFExporter/Source/GLTFExporter/Private/Converters/GLTFMaterialUtility.h
@@ -32,6 +32,8 @@ struct FGLTFMaterialUtility
 {
 	static UMaterialInterface* GetDefaultMaterial();
 
+	static bool IsClearCoatBottomNormalEnabled();
+
 #if WITH_EDITOR
 	static bool IsNormalMap(const FGLTFMaterialPropertyEx& Property);
 	static bool IsSRGB(const FGLTFMaterialPropertyEx& Property);
@@ -51,7 +53,7 @@ struct FGLTFMaterialUtility
 		return static_cast<const FMaterialInput<InputType>*>(ExpressionInput);
 	}
 
-	static const UMaterialExpressionCustomOutput* GetCustomOutputByName(const UMaterialInterface* Material, const FString& Name);
+	static const UMaterialExpressionCustomOutput* GetCustomOutputByName(const UMaterialInterface* Material, const FString& FunctionName);
 
 	static FGLTFPropertyBakeOutput BakeMaterialProperty(const FIntPoint& OutputSize, const FGLTFMaterialPropertyEx& Property, const UMaterialInterface* Material, int32 TexCoord, const FGLTFMeshData* MeshData = nullptr, const FGLTFIndexArray& MeshSectionIndices = {}, bool bFillAlpha = true, bool bAdjustNormalmaps = true);
 
diff --git a/Engine/Plugins/Enterprise/GLTFExporter/Source/GLTFExporter/Private/Tasks/GLTFDelayedMaterialTasks.cpp b/Engine/Plugins/Enterprise/GLTFExporter/Source/GLTFExporter/Private/Tasks/GLTFDelayedMaterialTasks.cpp
--- a/Engine/Plugins/Enterprise/GLTFExporter/Source/GLTFExporter/Private/Tasks/GLTFDelayedMaterialTasks.cpp
+++ b/Engine/Plugins/Enterprise/GLTFExporter/Source/GLTFExporter/Private/Tasks/GLTFDelayedMaterialTasks.cpp
@@ -225,8 +225,20 @@ void FGLTFDelayedMaterialTask::GetProxyParameters(FGLTFJsonMaterial& OutMaterial
 		GetProxyParameter(FGLTFProxyMaterialInfo::RoughnessFactor, OutMaterial.PBRMetallicRoughness.RoughnessFactor);
 		GetProxyParameter(FGLTFProxyMaterialInfo::MetallicRoughness, OutMaterial.PBRMetallicRoughness.MetallicRoughnessTexture);
 
-		GetProxyParameter(FGLTFProxyMaterialInfo::NormalScale, OutMaterial.NormalTexture.Scale);
-		GetProxyParameter(FGLTFProxyMaterialInfo::Normal, OutMaterial.NormalTexture);
+		if (HasProxyParameter(FGLTFProxyMaterialInfo::Normal.Texture))
+		{
+			if (OutMaterial.ShadingModel == EGLTFJsonShadingModel::ClearCoat && !FGLTFMaterialUtility::IsClearCoatBottomNormalEnabled())
+			{
+				Builder.LogWarning(FString::Printf(
+					TEXT("Proxy material %s won't be exported with (bottom) normal because ClearCoatEnableSecondNormal in project rendering settings is disabled"),
+					*Material->GetName()));
+			}
+			else
+			{
+				GetProxyParameter(FGLTFProxyMaterialInfo::NormalScale, OutMaterial.NormalTexture.Scale);
+				GetProxyParameter(FGLTFProxyMaterialInfo::Normal, OutMaterial.NormalTexture);
+			}
+		}
 
 		GetProxyParameter(FGLTFProxyMaterialInfo::OcclusionStrength, OutMaterial.OcclusionTexture.Strength);
 		GetProxyParameter(FGLTFProxyMaterialInfo::Occlusion, OutMaterial.OcclusionTexture);
@@ -856,17 +868,20 @@ bool FGLTFDelayedMaterialTask::TryGetEmissive(FGLTFJsonMaterial& OutMaterial, co
 
 bool FGLTFDelayedMaterialTask::IsPropertyNonDefault(const FGLTFMaterialPropertyEx& Property) const
 {
-	const bool bUseMaterialAttributes = Material->GetMaterial()->bUseMaterialAttributes;
-	if (bUseMaterialAttributes)
+	// Custom outputs are by definition standalone and never part of material attributes
+	if (!Property.IsCustomOutput())
 	{
-		// TODO: check if attribute property connected, i.e. Material->GetMaterial()->MaterialAttributes.IsConnected(Property)
-		return true;
+		const bool bUseMaterialAttributes = Material->GetMaterial()->bUseMaterialAttributes;
+		if (bUseMaterialAttributes)
+		{
+			// TODO: check if attribute property connected, i.e. Material->GetMaterial()->MaterialAttributes.IsConnected(Property)
+			return true;
+		}
 	}
 
 	const FExpressionInput* MaterialInput = FGLTFMaterialUtility::GetInputForProperty(Material, Property);
 	if (MaterialInput == nullptr)
 	{
-		// TODO: report error
 		return false;
 	}
 
@@ -876,6 +891,14 @@ bool FGLTFDelayedMaterialTask::IsPropertyNonDefault(const FGLTFMaterialPropertyE
 		return false;
 	}
 
+	if (Property == FGLTFMaterialPropertyEx::ClearCoatBottomNormal && !FGLTFMaterialUtility::IsClearCoatBottomNormalEnabled())
+	{
+		Builder.LogWarning(FString::Printf(
+			TEXT("Material %s won't be exported with clear coat bottom normal because ClearCoatEnableSecondNormal in project rendering settings is disabled"),
+			*Material->GetName()));
+		return false;
+	}
+
 	return true;
 }
 
diff --git a/Engine/Plugins/Enterprise/GLTFExporter/Source/GLTFExporter/Private/Tasks/GLTFDelayedMaterialTasks.h b/Engine/Plugins/Enterprise/GLTFExporter/Source/GLTFExporter/Private/Tasks/GLTFDelayedMaterialTasks.h
--- a/Engine/Plugins/Enterprise/GLTFExporter/Source/GLTFExporter/Private/Tasks/GLTFDelayedMaterialTasks.h
+++ b/Engine/Plugins/Enterprise/GLTFExporter/Source/GLTFExporter/Private/Tasks/GLTFDelayedMaterialTasks.h
@@ -54,6 +54,13 @@ private:
 	void GetProxyParameter(const TGLTFProxyMaterialParameterInfo<FLinearColor>& ParameterInfo, FGLTFJsonColor4& OutValue) const;
 	void GetProxyParameter(const FGLTFProxyMaterialTextureParameterInfo& ParameterInfo, FGLTFJsonTextureInfo& OutValue) const;
 
+	template <typename ParameterType>
+	bool HasProxyParameter(const TGLTFProxyMaterialParameterInfo<ParameterType>& ParameterInfo) const
+	{
+		ParameterType Value;
+		return ParameterInfo.Get(Material, Value, true);
+	}
+
 	EMaterialShadingModel GetShadingModel() const;
 	void ConvertShadingModel(EGLTFJsonShadingModel& OutShadingModel) const;
 	void ConvertAlphaMode(EGLTFJsonAlphaMode& OutAlphaMode) const;
