Rename all GLTF*Utility classes to GLTF*Utilities to conform with rest of engine

diff --git a/Engine/Plugins/Enterprise/GLTFExporter/Source/GLTFExporter/Private/Builders/GLTFBuilder.cpp b/Engine/Plugins/Enterprise/GLTFExporter/Source/GLTFExporter/Private/Builders/GLTFBuilder.cpp
--- a/Engine/Plugins/Enterprise/GLTFExporter/Source/GLTFExporter/Private/Builders/GLTFBuilder.cpp
+++ b/Engine/Plugins/Enterprise/GLTFExporter/Source/GLTFExporter/Private/Builders/GLTFBuilder.cpp
@@ -1,7 +1,7 @@
 // Copyright Epic Games, Inc. All Rights Reserved.
 
 #include "Builders/GLTFBuilder.h"
-#include "Converters/GLTFMeshUtility.h"
+#include "Converters/GLTFMeshUtilities.h"
 #include "Misc/Paths.h"
 #include "UserData/GLTFMaterialUserData.h"
 
@@ -79,15 +79,15 @@ TextureAddress FGLTFBuilder::GetBakeTilingForMaterialProperty(const UMaterialInt
 int32 FGLTFBuilder::SanitizeLOD(const UStaticMesh* StaticMesh, const UStaticMeshComponent* StaticMeshComponent, int32 LODIndex) const
 {
 	return LODIndex > 0
-		? FMath::Min(LODIndex, FGLTFMeshUtility::GetMaximumLOD(StaticMesh))
-		: FGLTFMeshUtility::GetLOD(StaticMesh, StaticMeshComponent, ExportOptions->DefaultLevelOfDetail);
+		? FMath::Min(LODIndex, FGLTFMeshUtilities::GetMaximumLOD(StaticMesh))
+		: FGLTFMeshUtilities::GetLOD(StaticMesh, StaticMeshComponent, ExportOptions->DefaultLevelOfDetail);
 }
 
 int32 FGLTFBuilder::SanitizeLOD(const USkeletalMesh* SkeletalMesh, const USkeletalMeshComponent* SkeletalMeshComponent, int32 LODIndex) const
 {
 	return LODIndex > 0
-		? FMath::Min(LODIndex, FGLTFMeshUtility::GetMaximumLOD(SkeletalMesh))
-		: FGLTFMeshUtility::GetLOD(SkeletalMesh, SkeletalMeshComponent, ExportOptions->DefaultLevelOfDetail);
+		? FMath::Min(LODIndex, FGLTFMeshUtilities::GetMaximumLOD(SkeletalMesh))
+		: FGLTFMeshUtilities::GetLOD(SkeletalMesh, SkeletalMeshComponent, ExportOptions->DefaultLevelOfDetail);
 }
 
 const UGLTFExportOptions* FGLTFBuilder::SanitizeExportOptions(const UGLTFExportOptions* Options)
diff --git a/Engine/Plugins/Enterprise/GLTFExporter/Source/GLTFExporter/Private/Builders/GLTFConvertBuilder.cpp b/Engine/Plugins/Enterprise/GLTFExporter/Source/GLTFExporter/Private/Builders/GLTFConvertBuilder.cpp
--- a/Engine/Plugins/Enterprise/GLTFExporter/Source/GLTFExporter/Private/Builders/GLTFConvertBuilder.cpp
+++ b/Engine/Plugins/Enterprise/GLTFExporter/Source/GLTFExporter/Private/Builders/GLTFConvertBuilder.cpp
@@ -1,7 +1,7 @@
 // Copyright Epic Games, Inc. All Rights Reserved.
 
 #include "Builders/GLTFConvertBuilder.h"
-#include "Converters/GLTFMeshUtility.h"
+#include "Converters/GLTFMeshUtilities.h"
 #include "Components/SkeletalMeshComponent.h"
 #include "Components/StaticMeshComponent.h"
 #include "Engine/Texture2D.h"
@@ -185,7 +185,7 @@ FGLTFJsonMaterial* FGLTFConvertBuilder::AddUniqueMaterial(const UMaterialInterfa
 {
 	// TODO: optimize by skipping mesh data if material doesn't need it
 	const FGLTFMeshData* MeshData = AddUniqueMeshData(StaticMesh, nullptr, LODIndex);
-	const FGLTFIndexArray SectionIndices = FGLTFMeshUtility::GetSectionIndices(StaticMesh, MeshData->LODIndex, MaterialIndex);
+	const FGLTFIndexArray SectionIndices = FGLTFMeshUtilities::GetSectionIndices(StaticMesh, MeshData->LODIndex, MaterialIndex);
 	return AddUniqueMaterial(Material, MeshData, SectionIndices);
 }
 
@@ -193,7 +193,7 @@ FGLTFJsonMaterial* FGLTFConvertBuilder::AddUniqueMaterial(const UMaterialInterfa
 {
 	// TODO: optimize by skipping mesh data if material doesn't need it
 	const FGLTFMeshData* MeshData = AddUniqueMeshData(SkeletalMesh, nullptr, LODIndex);
-	const FGLTFIndexArray SectionIndices = FGLTFMeshUtility::GetSectionIndices(SkeletalMesh, MeshData->LODIndex, MaterialIndex);
+	const FGLTFIndexArray SectionIndices = FGLTFMeshUtilities::GetSectionIndices(SkeletalMesh, MeshData->LODIndex, MaterialIndex);
 	return AddUniqueMaterial(Material, MeshData, SectionIndices);
 }
 
@@ -217,7 +217,7 @@ FGLTFJsonMaterial* FGLTFConvertBuilder::AddUniqueMaterial(const UMaterialInterfa
 	// TODO: optimize by skipping mesh data if material doesn't need it
 	const UStaticMesh* StaticMesh = StaticMeshComponent->GetStaticMesh();
 	const FGLTFMeshData* MeshData = AddUniqueMeshData(StaticMesh, StaticMeshComponent, LODIndex);
-	const FGLTFIndexArray SectionIndices = FGLTFMeshUtility::GetSectionIndices(StaticMesh, MeshData->LODIndex, MaterialIndex);
+	const FGLTFIndexArray SectionIndices = FGLTFMeshUtilities::GetSectionIndices(StaticMesh, MeshData->LODIndex, MaterialIndex);
 	return AddUniqueMaterial(Material, MeshData, SectionIndices);
 }
 
@@ -226,7 +226,7 @@ FGLTFJsonMaterial* FGLTFConvertBuilder::AddUniqueMaterial(const UMaterialInterfa
 	// TODO: optimize by skipping mesh data if material doesn't need it
 	const USkeletalMesh* SkeletalMesh = SkeletalMeshComponent->GetSkeletalMeshAsset();
 	const FGLTFMeshData* MeshData = AddUniqueMeshData(SkeletalMesh, SkeletalMeshComponent, LODIndex);
-	const FGLTFIndexArray SectionIndices = FGLTFMeshUtility::GetSectionIndices(SkeletalMesh, MeshData->LODIndex, MaterialIndex);
+	const FGLTFIndexArray SectionIndices = FGLTFMeshUtilities::GetSectionIndices(SkeletalMesh, MeshData->LODIndex, MaterialIndex);
 	return AddUniqueMaterial(Material, MeshData, SectionIndices);
 }
 
diff --git a/Engine/Plugins/Enterprise/GLTFExporter/Source/GLTFExporter/Private/Converters/GLTFBoneUtilities.cpp b/Engine/Plugins/Enterprise/GLTFExporter/Source/GLTFExporter/Private/Converters/GLTFBoneUtilities.cpp
new file mode 100644
--- /dev/null
+++ b/Engine/Plugins/Enterprise/GLTFExporter/Source/GLTFExporter/Private/Converters/GLTFBoneUtilities.cpp
@@ -0,0 +1,91 @@
+// Copyright Epic Games, Inc. All Rights Reserved.
+
+#include "Converters/GLTFBoneUtilities.h"
+#include "BonePose.h"
+#include "Sections/MovieScene3DTransformSection.h"
+#include "Animation/AnimationPoseData.h"
+#include "Animation/AnimSequence.h"
+#include "Animation/AttributesRuntime.h"
+#include "ReferenceSkeleton.h"
+
+FTransform FGLTFBoneUtilities::GetBindTransform(const FReferenceSkeleton& RefSkeleton, int32 BoneIndex)
+{
+	const TArray<FMeshBoneInfo>& BoneInfos = RefSkeleton.GetRefBoneInfo();
+	const TArray<FTransform>& BonePoses = RefSkeleton.GetRefBonePose();
+
+	int32 CurBoneIndex = BoneIndex;
+	FTransform BindTransform = FTransform::Identity;
+
+	do
+	{
+		BindTransform = BindTransform * BonePoses[CurBoneIndex];
+		CurBoneIndex = BoneInfos[CurBoneIndex].ParentIndex;
+	} while (CurBoneIndex != INDEX_NONE);
+
+	return BindTransform;
+}
+
+void FGLTFBoneUtilities::GetFrameTimestamps(const UAnimSequence* AnimSequence, TArray<float>& OutFrameTimestamps)
+{
+	const int32 FrameCount = AnimSequence->GetNumberOfSampledKeys();
+	OutFrameTimestamps.AddUninitialized(FrameCount);
+
+	const float SequenceLength = AnimSequence->GetPlayLength();
+	const float FrameLength = FrameCount > 1 ? SequenceLength / (FrameCount - 1) : 0;
+
+	for (int32 FrameIndex = 0; FrameIndex < FrameCount; ++FrameIndex)
+	{
+		OutFrameTimestamps[FrameIndex] = FMath::Clamp(FrameIndex * FrameLength, 0.0f, SequenceLength);
+	}
+}
+
+void FGLTFBoneUtilities::GetBoneIndices(const USkeleton* Skeleton, TArray<FBoneIndexType>& OutBoneIndices)
+{
+	const int32 BoneCount = Skeleton->GetReferenceSkeleton().GetNum();
+	OutBoneIndices.AddUninitialized(BoneCount);
+
+	for (int32 BoneIndex = 0; BoneIndex < BoneCount; ++BoneIndex)
+	{
+		OutBoneIndices[BoneIndex] = BoneIndex;
+	}
+}
+
+void FGLTFBoneUtilities::GetBoneTransformsByFrame(const UAnimSequence* AnimSequence, const TArray<float>& FrameTimestamps, const TArray<FBoneIndexType>& BoneIndices, TArray<TArray<FTransform>>& OutBoneTransformsByFrame)
+{
+#if WITH_EDITOR
+	const bool bUseRawData = AnimSequence->OnlyUseRawData();
+#else
+	const bool bUseRawData = false;
+#endif
+
+	// Make sure to free stack allocations made by FCompactPose, FBlendedCurve, and FStackCustomAttributes when end of scope
+	FMemMark Mark(FMemStack::Get());
+
+	FBoneContainer BoneContainer;
+	BoneContainer.InitializeTo(BoneIndices, FCurveEvaluationOption(true), *AnimSequence->GetSkeleton());
+	BoneContainer.SetUseRAWData(bUseRawData);
+
+	const int32 FrameCount = FrameTimestamps.Num();
+	OutBoneTransformsByFrame.AddDefaulted(FrameCount);
+
+	FCompactPose Pose;
+	Pose.SetBoneContainer(&BoneContainer);
+
+	FBlendedCurve Curve;
+	Curve.InitFrom(BoneContainer);
+
+	UE::Anim::FStackAttributeContainer Attributes;
+	FAnimationPoseData PoseData(Pose, Curve, Attributes);
+
+	for (int32 Frame = 0; Frame < FrameCount; ++Frame)
+	{
+		const FAnimExtractContext ExtractionContext(static_cast<double>(FrameTimestamps[Frame])); // TODO: set bExtractRootMotion?
+		AnimSequence->GetBonePose(PoseData, ExtractionContext, bUseRawData);
+		Pose.CopyBonesTo(OutBoneTransformsByFrame[Frame]);
+	}
+
+	// Clear all stack allocations to allow FMemMark to free them
+	Pose.Empty();
+	Curve.Empty();
+	Attributes.Empty();
+}
diff --git a/Engine/Plugins/Enterprise/GLTFExporter/Source/GLTFExporter/Private/Converters/GLTFBoneUtilities.h b/Engine/Plugins/Enterprise/GLTFExporter/Source/GLTFExporter/Private/Converters/GLTFBoneUtilities.h
new file mode 100644
--- /dev/null
+++ b/Engine/Plugins/Enterprise/GLTFExporter/Source/GLTFExporter/Private/Converters/GLTFBoneUtilities.h
@@ -0,0 +1,23 @@
+// Copyright Epic Games, Inc. All Rights Reserved.
+
+#pragma once
+
+#include "CoreMinimal.h"
+#include "BoneIndices.h"
+
+struct FMovieSceneDoubleChannel;
+struct FReferenceSkeleton;
+class UAnimSequence;
+class UMovieScene3DTransformSection;
+class USkeleton;
+
+struct FGLTFBoneUtilities
+{
+	static FTransform GetBindTransform(const FReferenceSkeleton& RefSkeleton, int32 BoneIndex);
+
+	static void GetFrameTimestamps(const UAnimSequence* AnimSequence, TArray<float>& OutFrameTimestamps);
+
+	static void GetBoneIndices(const USkeleton* Skeleton, TArray<FBoneIndexType>& OutBoneIndices);
+
+	static void GetBoneTransformsByFrame(const UAnimSequence* AnimSequence, const TArray<float>& FrameTimestamps, const TArray<FBoneIndexType>& BoneIndices, TArray<TArray<FTransform>>& OutBoneTransformsByFrame);
+};
diff --git a/Engine/Plugins/Enterprise/GLTFExporter/Source/GLTFExporter/Private/Converters/GLTFBoneUtility.cpp b/Engine/Plugins/Enterprise/GLTFExporter/Source/GLTFExporter/Private/Converters/GLTFBoneUtility.cpp
deleted file mode 100644
diff --git a/Engine/Plugins/Enterprise/GLTFExporter/Source/GLTFExporter/Private/Converters/GLTFBoneUtility.h b/Engine/Plugins/Enterprise/GLTFExporter/Source/GLTFExporter/Private/Converters/GLTFBoneUtility.h
deleted file mode 100644
diff --git a/Engine/Plugins/Enterprise/GLTFExporter/Source/GLTFExporter/Private/Converters/GLTFBufferAdapter.cpp b/Engine/Plugins/Enterprise/GLTFExporter/Source/GLTFExporter/Private/Converters/GLTFBufferAdapter.cpp
--- a/Engine/Plugins/Enterprise/GLTFExporter/Source/GLTFExporter/Private/Converters/GLTFBufferAdapter.cpp
+++ b/Engine/Plugins/Enterprise/GLTFExporter/Source/GLTFExporter/Private/Converters/GLTFBufferAdapter.cpp
@@ -1,7 +1,7 @@
 // Copyright Epic Games, Inc. All Rights Reserved.
 
 #include "Converters/GLTFBufferAdapter.h"
-#include "Converters/GLTFBufferUtility.h"
+#include "Converters/GLTFBufferUtilities.h"
 #include "RawIndexBuffer.h"
 #include "Rendering/ColorVertexBuffer.h"
 #include "Rendering/PositionVertexBuffer.h"
@@ -31,7 +31,7 @@ public:
 
 	FGLTFBufferAdapterGPU(FRHIBuffer* RHIBuffer)
 	{
-		FGLTFBufferUtility::ReadRHIBuffer(RHIBuffer, DataBuffer);
+		FGLTFBufferUtilities::ReadRHIBuffer(RHIBuffer, DataBuffer);
 	}
 
 	virtual const uint8* GetData() override
@@ -44,43 +44,43 @@ public:
 
 TUniquePtr<IGLTFBufferAdapter> IGLTFBufferAdapter::GetIndices(const FRawStaticIndexBuffer* IndexBuffer)
 {
-	const void* IndexData = FGLTFBufferUtility::GetCPUBuffer(IndexBuffer);
-	if (IndexData != nullptr && FGLTFBufferUtility::HasCPUAccess(IndexBuffer)) return MakeUnique<FGLTFBufferAdapterCPU>(IndexData);
+	const void* IndexData = FGLTFBufferUtilities::GetCPUBuffer(IndexBuffer);
+	if (IndexData != nullptr && FGLTFBufferUtilities::HasCPUAccess(IndexBuffer)) return MakeUnique<FGLTFBufferAdapterCPU>(IndexData);
 	return MakeUnique<FGLTFBufferAdapterGPU>(IndexBuffer->IndexBufferRHI);
 }
 
 TUniquePtr<IGLTFBufferAdapter> IGLTFBufferAdapter::GetIndices(const FRawStaticIndexBuffer16or32Interface* IndexBuffer)
 {
-	const void* IndexData =FGLTFBufferUtility::GetCPUBuffer(IndexBuffer);
-	if (IndexData != nullptr && FGLTFBufferUtility::HasCPUAccess(IndexBuffer)) return MakeUnique<FGLTFBufferAdapterCPU>(IndexData);
+	const void* IndexData =FGLTFBufferUtilities::GetCPUBuffer(IndexBuffer);
+	if (IndexData != nullptr && FGLTFBufferUtilities::HasCPUAccess(IndexBuffer)) return MakeUnique<FGLTFBufferAdapterCPU>(IndexData);
 	return MakeUnique<FGLTFBufferAdapterGPU>(IndexBuffer->IndexBufferRHI);
 }
 
 TUniquePtr<IGLTFBufferAdapter> IGLTFBufferAdapter::GetPositions(const FPositionVertexBuffer* VertexBuffer)
 {
 	const void* PositionData = VertexBuffer->GetVertexData();
-	if (PositionData != nullptr && FGLTFBufferUtility::HasCPUAccess(VertexBuffer)) return MakeUnique<FGLTFBufferAdapterCPU>(PositionData);
+	if (PositionData != nullptr && FGLTFBufferUtilities::HasCPUAccess(VertexBuffer)) return MakeUnique<FGLTFBufferAdapterCPU>(PositionData);
 	return MakeUnique<FGLTFBufferAdapterGPU>(VertexBuffer->VertexBufferRHI);
 }
 
 TUniquePtr<IGLTFBufferAdapter> IGLTFBufferAdapter::GetColors(const FColorVertexBuffer* VertexBuffer)
 {
 	const void* ColorData = VertexBuffer->GetVertexData();
-	if (ColorData != nullptr && FGLTFBufferUtility::HasCPUAccess(VertexBuffer)) return MakeUnique<FGLTFBufferAdapterCPU>(ColorData);
+	if (ColorData != nullptr && FGLTFBufferUtilities::HasCPUAccess(VertexBuffer)) return MakeUnique<FGLTFBufferAdapterCPU>(ColorData);
 	return MakeUnique<FGLTFBufferAdapterGPU>(VertexBuffer->VertexBufferRHI);
 }
 
 TUniquePtr<IGLTFBufferAdapter> IGLTFBufferAdapter::GetTangents(const FStaticMeshVertexBuffer* VertexBuffer)
 {
 	const void* TangentData = VertexBuffer->GetTangentData();
-	if (TangentData != nullptr && FGLTFBufferUtility::HasCPUAccess(VertexBuffer)) return MakeUnique<FGLTFBufferAdapterCPU>(TangentData);
+	if (TangentData != nullptr && FGLTFBufferUtilities::HasCPUAccess(VertexBuffer)) return MakeUnique<FGLTFBufferAdapterCPU>(TangentData);
 	return MakeUnique<FGLTFBufferAdapterGPU>(VertexBuffer->TangentsVertexBuffer.VertexBufferRHI);
 }
 
 TUniquePtr<IGLTFBufferAdapter> IGLTFBufferAdapter::GetUVs(const FStaticMeshVertexBuffer* VertexBuffer)
 {
 	const void* UVData = VertexBuffer->GetTexCoordData();
-	if (UVData != nullptr && FGLTFBufferUtility::HasCPUAccess(VertexBuffer)) return MakeUnique<FGLTFBufferAdapterCPU>(UVData);
+	if (UVData != nullptr && FGLTFBufferUtilities::HasCPUAccess(VertexBuffer)) return MakeUnique<FGLTFBufferAdapterCPU>(UVData);
 	return MakeUnique<FGLTFBufferAdapterGPU>(VertexBuffer->TexCoordVertexBuffer.VertexBufferRHI);
 }
 
@@ -88,7 +88,7 @@ TUniquePtr<IGLTFBufferAdapter> IGLTFBufferAdapter::GetInfluences(const FSkinWeig
 {
 	const FSkinWeightDataVertexBuffer* InfluenceBuffer = VertexBuffer->GetDataVertexBuffer();
 	const void* InfluenceData = InfluenceBuffer->GetWeightData();
-	if (InfluenceData != nullptr && FGLTFBufferUtility::HasCPUAccess(VertexBuffer)) return MakeUnique<FGLTFBufferAdapterCPU>(InfluenceData);
+	if (InfluenceData != nullptr && FGLTFBufferUtilities::HasCPUAccess(VertexBuffer)) return MakeUnique<FGLTFBufferAdapterCPU>(InfluenceData);
 	return MakeUnique<FGLTFBufferAdapterGPU>(InfluenceBuffer->VertexBufferRHI);
 }
 
@@ -96,6 +96,6 @@ TUniquePtr<IGLTFBufferAdapter> IGLTFBufferAdapter::GetLookups(const FSkinWeightV
 {
 	const FSkinWeightLookupVertexBuffer* LookupBuffer = VertexBuffer->GetLookupVertexBuffer();
 	const void* LookupData = LookupBuffer->GetLookupData();
-	if (LookupData != nullptr && FGLTFBufferUtility::HasCPUAccess(VertexBuffer)) return MakeUnique<FGLTFBufferAdapterCPU>(LookupData);
+	if (LookupData != nullptr && FGLTFBufferUtilities::HasCPUAccess(VertexBuffer)) return MakeUnique<FGLTFBufferAdapterCPU>(LookupData);
 	return MakeUnique<FGLTFBufferAdapterGPU>(LookupBuffer->VertexBufferRHI);
 }
diff --git a/Engine/Plugins/Enterprise/GLTFExporter/Source/GLTFExporter/Private/Converters/GLTFBufferUtilities.cpp b/Engine/Plugins/Enterprise/GLTFExporter/Source/GLTFExporter/Private/Converters/GLTFBufferUtilities.cpp
new file mode 100644
--- /dev/null
+++ b/Engine/Plugins/Enterprise/GLTFExporter/Source/GLTFExporter/Private/Converters/GLTFBufferUtilities.cpp
@@ -0,0 +1,104 @@
+// Copyright Epic Games, Inc. All Rights Reserved.
+
+#include "Converters/GLTFBufferUtilities.h"
+#include "Rendering/SkeletalMeshRenderData.h"
+
+bool FGLTFBufferUtilities::HasCPUAccess(const FRawStaticIndexBuffer* IndexBuffer)
+{
+#if WITH_EDITOR
+	return true;
+#else
+	return IndexBuffer->GetAllowCPUAccess();
+#endif
+}
+
+bool FGLTFBufferUtilities::HasCPUAccess(const FRawStaticIndexBuffer16or32Interface* IndexBuffer)
+{
+#if WITH_EDITOR
+	return true;
+#else
+	return IndexBuffer->GetNeedsCPUAccess();
+#endif
+}
+
+bool FGLTFBufferUtilities::HasCPUAccess(const FPositionVertexBuffer* VertexBuffer)
+{
+#if WITH_EDITOR
+	return true;
+#else
+	return VertexBuffer->GetAllowCPUAccess();
+#endif
+}
+
+bool FGLTFBufferUtilities::HasCPUAccess(const FColorVertexBuffer* VertexBuffer)
+{
+#if WITH_EDITOR
+	return true;
+#else
+	return VertexBuffer->GetAllowCPUAccess();
+#endif
+}
+
+bool FGLTFBufferUtilities::HasCPUAccess(const FStaticMeshVertexBuffer* VertexBuffer)
+{
+#if WITH_EDITOR
+	return true;
+#else
+	return const_cast<FStaticMeshVertexBuffer*>(VertexBuffer)->GetAllowCPUAccess();
+#endif
+}
+
+bool FGLTFBufferUtilities::HasCPUAccess(const FSkinWeightVertexBuffer* VertexBuffer)
+{
+#if WITH_EDITOR
+	return true;
+#else
+	return VertexBuffer->GetNeedsCPUAccess();
+#endif
+}
+
+const void* FGLTFBufferUtilities::GetCPUBuffer(const FRawStaticIndexBuffer* IndexBuffer)
+{
+	return IndexBuffer->Is32Bit() ? static_cast<const void*>(IndexBuffer->AccessStream32()) : static_cast<const void*>(IndexBuffer->AccessStream16());
+}
+
+const void* FGLTFBufferUtilities::GetCPUBuffer(const FRawStaticIndexBuffer16or32Interface* IndexBuffer)
+{
+	return IndexBuffer->GetResourceDataSize() > 0 ? const_cast<FRawStaticIndexBuffer16or32Interface*>(IndexBuffer)->GetPointerTo(0) : nullptr;
+}
+
+void FGLTFBufferUtilities::ReadRHIBuffer(FRHIBuffer* SourceBuffer, TArray<uint8>& OutData)
+{
+	OutData.Empty();
+
+	if (SourceBuffer == nullptr)
+	{
+		return;
+	}
+
+	const uint32 NumBytes = SourceBuffer->GetSize();
+	if (NumBytes == 0)
+	{
+		return;
+	}
+
+	const EBufferUsageFlags Usage = SourceBuffer->GetUsage();
+	if ((Usage & BUF_Static) != BUF_Static)
+	{
+		return; // Some RHI implementations only support reading static buffers
+	}
+
+	OutData.AddUninitialized(NumBytes);
+	void *DstData = OutData.GetData();
+
+	ENQUEUE_RENDER_COMMAND(ReadRHIBuffer)(
+		[SourceBuffer, NumBytes, DstData](FRHICommandListImmediate& RHICmdList)
+		{
+			const void* SrcData = RHICmdList.LockBuffer(SourceBuffer, 0, NumBytes, RLM_ReadOnly);
+			FMemory::Memcpy(DstData, SrcData, NumBytes);
+			RHICmdList.UnlockBuffer(SourceBuffer);
+		}
+	);
+
+	FlushRenderingCommands();
+}
diff --git a/Engine/Plugins/Enterprise/GLTFExporter/Source/GLTFExporter/Private/Converters/GLTFBufferUtilities.h b/Engine/Plugins/Enterprise/GLTFExporter/Source/GLTFExporter/Private/Converters/GLTFBufferUtilities.h
new file mode 100644
--- /dev/null
+++ b/Engine/Plugins/Enterprise/GLTFExporter/Source/GLTFExporter/Private/Converters/GLTFBufferUtilities.h
@@ -0,0 +1,30 @@
+// Copyright Epic Games, Inc. All Rights Reserved.
+
+#pragma once
+
+#include "CoreMinimal.h"
+
+class FColorVertexBuffer;
+class FPositionVertexBuffer;
+class FRawStaticIndexBuffer;
+class FRawStaticIndexBuffer16or32Interface;
+class FRHIBuffer;
+class FSkinWeightVertexBuffer;
+class FSkinWeightDataVertexBuffer;
+class FSkinWeightLookupVertexBuffer;
+class FStaticMeshVertexBuffer;
+
+struct FGLTFBufferUtilities
+{
+	static bool HasCPUAccess(const FRawStaticIndexBuffer* IndexBuffer);
+	static bool HasCPUAccess(const FRawStaticIndexBuffer16or32Interface* IndexBuffer);
+	static bool HasCPUAccess(const FPositionVertexBuffer* VertexBuffer);
+	static bool HasCPUAccess(const FColorVertexBuffer* VertexBuffer);
+	static bool HasCPUAccess(const FStaticMeshVertexBuffer* VertexBuffer);
+	static bool HasCPUAccess(const FSkinWeightVertexBuffer* VertexBuffer);
+
+	static const void* GetCPUBuffer(const FRawStaticIndexBuffer* IndexBuffer);
+	static const void* GetCPUBuffer(const FRawStaticIndexBuffer16or32Interface* IndexBuffer);
+
+	static void ReadRHIBuffer(FRHIBuffer* SourceBuffer, TArray<uint8>& OutData);
+};
diff --git a/Engine/Plugins/Enterprise/GLTFExporter/Source/GLTFExporter/Private/Converters/GLTFBufferUtility.cpp b/Engine/Plugins/Enterprise/GLTFExporter/Source/GLTFExporter/Private/Converters/GLTFBufferUtility.cpp
deleted file mode 100644
diff --git a/Engine/Plugins/Enterprise/GLTFExporter/Source/GLTFExporter/Private/Converters/GLTFBufferUtility.h b/Engine/Plugins/Enterprise/GLTFExporter/Source/GLTFExporter/Private/Converters/GLTFBufferUtility.h
deleted file mode 100644
diff --git a/Engine/Plugins/Enterprise/GLTFExporter/Source/GLTFExporter/Private/Converters/GLTFCameraConverters.cpp b/Engine/Plugins/Enterprise/GLTFExporter/Source/GLTFExporter/Private/Converters/GLTFCameraConverters.cpp
--- a/Engine/Plugins/Enterprise/GLTFExporter/Source/GLTFExporter/Private/Converters/GLTFCameraConverters.cpp
+++ b/Engine/Plugins/Enterprise/GLTFExporter/Source/GLTFExporter/Private/Converters/GLTFCameraConverters.cpp
@@ -3,7 +3,7 @@
 #include "Converters/GLTFCameraConverters.h"
 #include "Builders/GLTFContainerBuilder.h"
 #include "Utilities/GLTFCoreUtilities.h"
-#include "Converters/GLTFNameUtility.h"
+#include "Converters/GLTFNameUtilities.h"
 #include "Camera/CameraComponent.h"
 #include "GameFramework/Actor.h"
 
@@ -17,7 +17,7 @@ FGLTFJsonCamera* FGLTFCameraConverter::Convert(const UCameraComponent* CameraCom
 	}
 
 	FGLTFJsonCamera* JsonCamera = Builder.AddCamera();
-	JsonCamera->Name = FGLTFNameUtility::GetName(CameraComponent);
+	JsonCamera->Name = FGLTFNameUtilities::GetName(CameraComponent);
 	JsonCamera->Type = Type;
 
 	FMinimalViewInfo DesiredView;
diff --git a/Engine/Plugins/Enterprise/GLTFExporter/Source/GLTFExporter/Private/Converters/GLTFImageConverters.cpp b/Engine/Plugins/Enterprise/GLTFExporter/Source/GLTFExporter/Private/Converters/GLTFImageConverters.cpp
--- a/Engine/Plugins/Enterprise/GLTFExporter/Source/GLTFExporter/Private/Converters/GLTFImageConverters.cpp
+++ b/Engine/Plugins/Enterprise/GLTFExporter/Source/GLTFExporter/Private/Converters/GLTFImageConverters.cpp
@@ -1,7 +1,7 @@
 // Copyright Epic Games, Inc. All Rights Reserved.
 
 #include "Converters/GLTFImageConverters.h"
-#include "Converters/GLTFImageUtility.h"
+#include "Converters/GLTFImageUtilities.h"
 #include "Builders/GLTFConvertBuilder.h"
 #include "Misc/FileHelper.h"
 
@@ -16,11 +16,11 @@ FGLTFJsonImage* FGLTFImageConverter::Convert(TGLTFSuperfluous<FString> Name, boo
 			return nullptr;
 
 		case EGLTFJsonMimeType::PNG:
-			FGLTFImageUtility::CompressToPNG(Pixels->GetData(), Size, *CompressedData);
+			FGLTFImageUtilities::CompressToPNG(Pixels->GetData(), Size, *CompressedData);
 			break;
 
 		case EGLTFJsonMimeType::JPEG:
-			FGLTFImageUtility::CompressToJPEG(Pixels->GetData(), Size, Builder.ExportOptions->TextureImageQuality, *CompressedData);
+			FGLTFImageUtilities::CompressToJPEG(Pixels->GetData(), Size, Builder.ExportOptions->TextureImageQuality, *CompressedData);
 			break;
 
 		default:
@@ -38,7 +38,7 @@ FGLTFJsonImage* FGLTFImageConverter::Convert(TGLTFSuperfluous<FString> Name, boo
 	}
 	else
 	{
-		const TCHAR* Extension = FGLTFImageUtility::GetFileExtension(MimeType);
+		const TCHAR* Extension = FGLTFImageUtilities::GetFileExtension(MimeType);
 		JsonImage->URI = Builder.AddExternalFile(Name + Extension, CompressedData);
 	}
 
@@ -51,7 +51,7 @@ EGLTFJsonMimeType FGLTFImageConverter::GetMimeType(const FColor* Pixels, FIntPoi
 	{
 		case EGLTFTextureImageFormat::None: return EGLTFJsonMimeType::None;
 		case EGLTFTextureImageFormat::PNG: return EGLTFJsonMimeType::PNG;
-		case EGLTFTextureImageFormat::JPEG: return bIgnoreAlpha || FGLTFImageUtility::NoAlphaNeeded(Pixels, Size) ? EGLTFJsonMimeType::JPEG : EGLTFJsonMimeType::PNG;
+		case EGLTFTextureImageFormat::JPEG: return bIgnoreAlpha || FGLTFImageUtilities::NoAlphaNeeded(Pixels, Size) ? EGLTFJsonMimeType::JPEG : EGLTFJsonMimeType::PNG;
 		default:
 			checkNoEntry();
 			return EGLTFJsonMimeType::None;
diff --git a/Engine/Plugins/Enterprise/GLTFExporter/Source/GLTFExporter/Private/Converters/GLTFImageUtilities.cpp b/Engine/Plugins/Enterprise/GLTFExporter/Source/GLTFExporter/Private/Converters/GLTFImageUtilities.cpp
new file mode 100644
--- /dev/null
+++ b/Engine/Plugins/Enterprise/GLTFExporter/Source/GLTFExporter/Private/Converters/GLTFImageUtilities.cpp
@@ -0,0 +1,84 @@
+// Copyright Epic Games, Inc. All Rights Reserved.
+
+#include "Converters/GLTFImageUtilities.h"
+#include "Modules/ModuleManager.h"
+#include "IImageWrapperModule.h"
+#include "IImageWrapper.h"
+
+const TCHAR* FGLTFImageUtilities::GetFileExtension(EGLTFJsonMimeType MimeType)
+{
+	switch (MimeType)
+	{
+		case EGLTFJsonMimeType::PNG:  return TEXT(".png");
+		case EGLTFJsonMimeType::JPEG: return TEXT(".jpg");
+		default:
+			checkNoEntry();
+		return TEXT("");
+	}
+}
+
+FString FGLTFImageUtilities::GetUniqueFilename(const FString& BaseFilename, const FString& FileExtension, const TSet<FString>& UniqueFilenames)
+{
+	FString Filename = BaseFilename + FileExtension;
+	if (!UniqueFilenames.Contains(Filename))
+	{
+		return Filename;
+	}
+
+	int32 Suffix = 1;
+	do
+	{
+		Filename = BaseFilename + TEXT('_') + FString::FromInt(Suffix) + FileExtension;
+		Suffix++;
+	}
+	while (UniqueFilenames.Contains(Filename));
+
+	return Filename;
+}
+
+bool FGLTFImageUtilities::NoAlphaNeeded(const FColor* Pixels, FIntPoint Size)
+{
+	const int64 Count = Size.X * Size.Y;
+	for (int64 Index = 0; Index < Count; ++Index)
+	{
+		const FColor& Pixel = Pixels[Index];
+		if (Pixel.A < 255)
+		{
+			return false;
+		}
+	}
+
+	return true;
+}
+
+bool FGLTFImageUtilities::CompressToPNG(const FColor* InPixels, FIntPoint InSize, TArray64<uint8>& OutCompressedData)
+{
+	return CompressToFormat(InPixels, InSize, EImageFormat::PNG, 0, OutCompressedData);
+}
+
+bool FGLTFImageUtilities::CompressToJPEG(const FColor* InPixels, FIntPoint InSize, int32 InCompressionQuality, TArray64<uint8>& OutCompressedData)
+{
+	// NOTE: Tests have shown that explicitly checking for greyscale images (and using EImageFormat::GrayscaleJPEG) results in the same as EImageFormat::JPEG
+	return CompressToFormat(InPixels, InSize, EImageFormat::JPEG, InCompressionQuality, OutCompressedData);
+}
+
+bool FGLTFImageUtilities::CompressToFormat(const FColor* InPixels, FIntPoint InSize, EImageFormat InCompressionFormat, int32 InCompressionQuality, TArray64<uint8>& OutCompressedData)
+{
+	const int64 ByteLength = InSize.X * InSize.Y * sizeof(FColor);
+	return CompressToFormat(InPixels, ByteLength, InSize.X, InSize.Y, ERGBFormat::BGRA, 8, InCompressionFormat, InCompressionQuality, OutCompressedData);
+}
+
+bool FGLTFImageUtilities::CompressToFormat(const void* InRawData, int64 InRawSize, int32 InWidth, int32 InHeight, ERGBFormat InRGBFormat, int32 InBitDepth, EImageFormat InCompressionFormat, int32 InCompressionQuality, TArray64<uint8>& OutCompressedData)
+{
+	IImageWrapperModule& ImageWrapperModule = FModuleManager::Get().LoadModuleChecked<IImageWrapperModule>(TEXT("ImageWrapper"));
+	const TSharedPtr<IImageWrapper> ImageWrapper = ImageWrapperModule.CreateImageWrapper(InCompressionFormat);
+
+	if (ImageWrapper.IsValid() && ImageWrapper->SetRaw(InRawData, InRawSize, InWidth, InHeight, InRGBFormat, InBitDepth))
+	{
+		const TArray64<uint8>& CompressedData = ImageWrapper->GetCompressed(InCompressionQuality);
+		OutCompressedData = MoveTemp(const_cast<TArray64<uint8>&>(CompressedData)); // safe to do this since we don't keep image wrapper
+		return true;
+	}
+
+	return false;
+}
diff --git a/Engine/Plugins/Enterprise/GLTFExporter/Source/GLTFExporter/Private/Converters/GLTFImageUtilities.h b/Engine/Plugins/Enterprise/GLTFExporter/Source/GLTFExporter/Private/Converters/GLTFImageUtilities.h
new file mode 100644
--- /dev/null
+++ b/Engine/Plugins/Enterprise/GLTFExporter/Source/GLTFExporter/Private/Converters/GLTFImageUtilities.h
@@ -0,0 +1,24 @@
+// Copyright Epic Games, Inc. All Rights Reserved.
+
+#pragma once
+
+#include "CoreMinimal.h"
+#include "Json/GLTFJsonEnums.h"
+
+enum class ERGBFormat : int8;
+enum class EImageFormat : int8;
+
+struct FGLTFImageUtilities
+{
+	static const TCHAR* GetFileExtension(EGLTFJsonMimeType MimeType);
+
+	static FString GetUniqueFilename(const FString& BaseFilename, const FString& FileExtension, const TSet<FString>& UniqueFilenames);
+
+	static bool NoAlphaNeeded(const FColor* Pixels, FIntPoint Size);
+
+	static bool CompressToPNG(const FColor* InPixels, FIntPoint InSize, TArray64<uint8>& OutCompressedData);
+	static bool CompressToJPEG(const FColor* InPixels, FIntPoint InSize, int32 InCompressionQuality, TArray64<uint8>& OutCompressedData);
+
+	static bool CompressToFormat(const FColor* InPixels, FIntPoint InSize, EImageFormat InCompressionFormat, int32 InCompressionQuality, TArray64<uint8>& OutCompressedData);
+	static bool CompressToFormat(const void* InRawData, int64 InRawSize, int32 InWidth, int32 InHeight, ERGBFormat InRGBFormat, int32 InBitDepth, EImageFormat InCompressionFormat, int32 InCompressionQuality, TArray64<uint8>& OutCompressedData);
+};
diff --git a/Engine/Plugins/Enterprise/GLTFExporter/Source/GLTFExporter/Private/Converters/GLTFImageUtility.cpp b/Engine/Plugins/Enterprise/GLTFExporter/Source/GLTFExporter/Private/Converters/GLTFImageUtility.cpp
deleted file mode 100644
diff --git a/Engine/Plugins/Enterprise/GLTFExporter/Source/GLTFExporter/Private/Converters/GLTFImageUtility.h b/Engine/Plugins/Enterprise/GLTFExporter/Source/GLTFExporter/Private/Converters/GLTFImageUtility.h
deleted file mode 100644
diff --git a/Engine/Plugins/Enterprise/GLTFExporter/Source/GLTFExporter/Private/Converters/GLTFLightConverters.cpp b/Engine/Plugins/Enterprise/GLTFExporter/Source/GLTFExporter/Private/Converters/GLTFLightConverters.cpp
--- a/Engine/Plugins/Enterprise/GLTFExporter/Source/GLTFExporter/Private/Converters/GLTFLightConverters.cpp
+++ b/Engine/Plugins/Enterprise/GLTFExporter/Source/GLTFExporter/Private/Converters/GLTFLightConverters.cpp
@@ -3,7 +3,7 @@
 #include "Converters/GLTFLightConverters.h"
 #include "Builders/GLTFContainerBuilder.h"
 #include "Utilities/GLTFCoreUtilities.h"
-#include "Converters/GLTFNameUtility.h"
+#include "Converters/GLTFNameUtilities.h"
 #include "Components/LightComponent.h"
 #include "Components/PointLightComponent.h"
 #include "Components/SpotLightComponent.h"
@@ -18,7 +18,7 @@ FGLTFJsonLight* FGLTFLightConverter::Convert(const ULightComponent* LightCompone
 	}
 
 	FGLTFJsonLight* Light = Builder.AddLight();
-	Light->Name = FGLTFNameUtility::GetName(LightComponent);
+	Light->Name = FGLTFNameUtilities::GetName(LightComponent);
 	Light->Type = Type;
 
 	const float ConversionScale = Type == EGLTFJsonLightType::Directional ? 1.0f : (0.01f * 0.01f); // Conversion from cm2 to m2
diff --git a/Engine/Plugins/Enterprise/GLTFExporter/Source/GLTFExporter/Private/Converters/GLTFMaterialConverters.cpp b/Engine/Plugins/Enterprise/GLTFExporter/Source/GLTFExporter/Private/Converters/GLTFMaterialConverters.cpp
--- a/Engine/Plugins/Enterprise/GLTFExporter/Source/GLTFExporter/Private/Converters/GLTFMaterialConverters.cpp
+++ b/Engine/Plugins/Enterprise/GLTFExporter/Source/GLTFExporter/Private/Converters/GLTFMaterialConverters.cpp
@@ -1,7 +1,7 @@
 // Copyright Epic Games, Inc. All Rights Reserved.
 
 #include "Converters/GLTFMaterialConverters.h"
-#include "Converters/GLTFMaterialUtility.h"
+#include "Converters/GLTFMaterialUtilities.h"
 #include "Builders/GLTFConvertBuilder.h"
 #include "Tasks/GLTFDelayedMaterialTasks.h"
 
@@ -11,7 +11,7 @@ void FGLTFMaterialConverter::Sanitize(const UMaterialInterface*& Material, const
 
 	if (MeshData == nullptr ||
 		Builder.ExportOptions->BakeMaterialInputs != EGLTFMaterialBakeMode::UseMeshData ||
-		!FGLTFMaterialUtility::NeedsMeshData(Material))
+		!FGLTFMaterialUtilities::NeedsMeshData(Material))
 	{
 		MeshData = nullptr;
 		SectionIndices = {};
@@ -45,7 +45,7 @@ void FGLTFMaterialConverter::Sanitize(const UMaterialInterface*& Material, const
 
 FGLTFJsonMaterial* FGLTFMaterialConverter::Convert(const UMaterialInterface* Material, const FGLTFMeshData* MeshData, FGLTFIndexArray SectionIndices)
 {
-	if (Material != FGLTFMaterialUtility::GetDefaultMaterial())
+	if (Material != FGLTFMaterialUtilities::GetDefaultMaterial())
 	{
 		FGLTFJsonMaterial* JsonMaterial = Builder.AddMaterial();
 		Builder.ScheduleSlowTask<FGLTFDelayedMaterialTask>(Builder, UVOverlapChecker, Material, MeshData, SectionIndices, JsonMaterial);
diff --git a/Engine/Plugins/Enterprise/GLTFExporter/Source/GLTFExporter/Private/Converters/GLTFMaterialProxyFactory.cpp b/Engine/Plugins/Enterprise/GLTFExporter/Source/GLTFExporter/Private/Converters/GLTFMaterialProxyFactory.cpp
--- a/Engine/Plugins/Enterprise/GLTFExporter/Source/GLTFExporter/Private/Converters/GLTFMaterialProxyFactory.cpp
+++ b/Engine/Plugins/Enterprise/GLTFExporter/Source/GLTFExporter/Private/Converters/GLTFMaterialProxyFactory.cpp
@@ -3,8 +3,8 @@
 #if WITH_EDITOR
 
 #include "Converters/GLTFMaterialProxyFactory.h"
-#include "Converters/GLTFMaterialUtility.h"
-#include "Converters/GLTFImageUtility.h"
+#include "Converters/GLTFMaterialUtilities.h"
+#include "Converters/GLTFImageUtilities.h"
 #include "Utilities/GLTFProxyMaterialUtilities.h"
 #include "Materials/GLTFProxyMaterialInfo.h"
 #include "Materials/Material.h"
@@ -28,7 +28,7 @@ UMaterialInterface* FGLTFMaterialProxyFactory::Create(UMaterialInterface* Origin
 		return OriginalMaterial;
 	}
 
-	if (FGLTFMaterialUtility::NeedsMeshData(OriginalMaterial))
+	if (FGLTFMaterialUtilities::NeedsMeshData(OriginalMaterial))
 	{
 		Builder.LogWarning(FString::Printf(
 			TEXT("Material %s uses mesh data but prebaking will only use a simple quad as mesh data currently"),
@@ -278,7 +278,7 @@ TUniquePtr<IGLTFImageConverter> FGLTFMaterialProxyFactory::CreateImageConverter(
 
 		virtual FGLTFJsonImage* Convert(TGLTFSuperfluous<FString> Name, bool bIgnoreAlpha, FIntPoint Size, TGLTFSharedArray<FColor> Pixels) override
 		{
-			const FString Filename = FGLTFImageUtility::GetUniqueFilename(Name, TEXT(""), UniqueFilenames);
+			const FString Filename = FGLTFImageUtilities::GetUniqueFilename(Name, TEXT(""), UniqueFilenames);
 			UniqueFilenames.Add(Filename);
 
 			FGLTFJsonImage* Image = Factory.Builder.AddImage();
diff --git a/Engine/Plugins/Enterprise/GLTFExporter/Source/GLTFExporter/Private/Converters/GLTFMaterialUtilities.cpp b/Engine/Plugins/Enterprise/GLTFExporter/Source/GLTFExporter/Private/Converters/GLTFMaterialUtilities.cpp
new file mode 100644
--- /dev/null
+++ b/Engine/Plugins/Enterprise/GLTFExporter/Source/GLTFExporter/Private/Converters/GLTFMaterialUtilities.cpp
@@ -0,0 +1,368 @@
+// Copyright Epic Games, Inc. All Rights Reserved.
+
+#include "Converters/GLTFMaterialUtilities.h"
+#include "Converters/GLTFTextureUtilities.h"
+#include "Converters/GLTFNameUtilities.h"
+#include "Materials/Material.h"
+#include "Utilities/GLTFProxyMaterialUtilities.h"
+#include "Misc/DefaultValueHelper.h"
+#if WITH_EDITOR
+#include "GLTFMaterialAnalyzer.h"
+#include "IGLTFMaterialBakingModule.h"
+#include "GLTFMaterialBakingStructures.h"
+#include "Materials/MaterialExpressionTextureSample.h"
+#include "Materials/MaterialAttributeDefinitionMap.h"
+#endif
+#include "Engine/RendererSettings.h"
+#include "Modules/ModuleManager.h"
+#include "Materials/MaterialInstanceConstant.h"
+#include "Materials/MaterialExpressionCustomOutput.h"
+#include "Materials/MaterialExpressionTextureCoordinate.h"
+
+UMaterialInterface* FGLTFMaterialUtilities::GetDefaultMaterial()
+{
+	static UMaterial* DefaultMaterial = FGLTFProxyMaterialUtilities::GetBaseMaterial(EGLTFJsonShadingModel::Default);
+	return DefaultMaterial;
+}
+
+bool FGLTFMaterialUtilities::IsClearCoatBottomNormalEnabled()
+{
+	return GetDefault<URendererSettings>()->bClearCoatEnableSecondNormal != 0;
+}
+
+#if WITH_EDITOR
+
+bool FGLTFMaterialUtilities::IsNormalMap(const FGLTFMaterialPropertyEx& Property)
+{
+	return Property == MP_Normal || Property == FGLTFMaterialPropertyEx::ClearCoatBottomNormal;
+}
+
+bool FGLTFMaterialUtilities::IsSRGB(const FGLTFMaterialPropertyEx& Property)
+{
+	return Property == MP_BaseColor || Property == MP_EmissiveColor || Property == MP_SubsurfaceColor || Property == FGLTFMaterialPropertyEx::TransmittanceColor;
+}
+
+FGuid FGLTFMaterialUtilities::GetAttributeID(const FGLTFMaterialPropertyEx& Property)
+{
+	return Property.IsCustomOutput()
+		? FMaterialAttributeDefinitionMap::GetCustomAttributeID(Property.CustomOutput.ToString())
+		: FMaterialAttributeDefinitionMap::GetID(Property.Type);
+}
+
+FGuid FGLTFMaterialUtilities::GetAttributeIDChecked(const FGLTFMaterialPropertyEx& Property)
+{
+	const FGuid AttributeID = GetAttributeID(Property);
+	check(AttributeID != FMaterialAttributeDefinitionMap::GetDefaultID());
+	return AttributeID;
+}
+
+FVector4f FGLTFMaterialUtilities::GetPropertyDefaultValue(const FGLTFMaterialPropertyEx& Property)
+{
+	return FMaterialAttributeDefinitionMap::GetDefaultValue(GetAttributeIDChecked(Property));
+}
+
+FVector4f FGLTFMaterialUtilities::GetPropertyMask(const FGLTFMaterialPropertyEx& Property)
+{
+	switch (FMaterialAttributeDefinitionMap::GetValueType(GetAttributeIDChecked(Property)))
+	{
+		case MCT_Float:
+		case MCT_Float1: return FVector4f(1, 0, 0, 0);
+		case MCT_Float2: return FVector4f(1, 1, 0, 0);
+		case MCT_Float3: return FVector4f(1, 1, 1, 0);
+		case MCT_Float4: return FVector4f(1, 1, 1, 1);
+		default:
+			checkNoEntry();
+			return FVector4f();
+	}
+}
+
+const FExpressionInput* FGLTFMaterialUtilities::GetInputForProperty(const UMaterialInterface* Material, const FGLTFMaterialPropertyEx& Property)
+{
+	if (Property.IsCustomOutput())
+	{
+		const FString FunctionName = Property.CustomOutput.ToString();
+		const UMaterialExpressionCustomOutput* CustomOutput = GetCustomOutputByName(Material, FunctionName);
+		if (CustomOutput == nullptr)
+		{
+			return nullptr;
+		}
+
+		// Assume custom outputs always have a single input (which is true for all supported custom outputs)
+		return const_cast<UMaterialExpressionCustomOutput*>(CustomOutput)->GetInput(0);
+	}
+
+	UMaterial* UnderlyingMaterial = const_cast<UMaterial*>(Material->GetMaterial());
+	return UnderlyingMaterial->GetExpressionInputForProperty(Property.Type);
+}
+
+const UMaterialExpressionCustomOutput* FGLTFMaterialUtilities::GetCustomOutputByName(const UMaterialInterface* Material, const FString& FunctionName)
+{
+	for (const TObjectPtr<UMaterialExpression>& Expression : Material->GetMaterial()->GetExpressions())
+	{
+		const UMaterialExpressionCustomOutput* CustomOutput = Cast<UMaterialExpressionCustomOutput>(Expression);
+		if (CustomOutput != nullptr && CustomOutput->GetFunctionName() == FunctionName)
+		{
+			return CustomOutput;
+		}
+	}
+
+	return nullptr;
+}
+
+FGLTFPropertyBakeOutput FGLTFMaterialUtilities::BakeMaterialProperty(const FIntPoint& OutputSize, const FGLTFMaterialPropertyEx& Property, const UMaterialInterface* Material, int32 TexCoord, const FGLTFMeshData* MeshData, const FGLTFIndexArray& MeshSectionIndices, bool bFillAlpha, bool bAdjustNormalmaps)
+{
+	FGLTFMeshRenderData MeshSet;
+	MeshSet.TextureCoordinateBox = { { 0.0f, 0.0f }, { 1.0f, 1.0f } };
+	MeshSet.TextureCoordinateIndex = TexCoord;
+	MeshSet.MaterialIndices = MeshSectionIndices; // NOTE: MaterialIndices is actually section indices
+	if (MeshData != nullptr)
+	{
+		MeshSet.MeshDescription = const_cast<FMeshDescription*>(&MeshData->Description);
+		MeshSet.LightMap = MeshData->LightMap;
+		MeshSet.LightMapIndex = MeshData->LightMapTexCoord;
+		MeshSet.LightmapResourceCluster = MeshData->LightMapResourceCluster;
+		MeshSet.PrimitiveData = &MeshData->PrimitiveData;
+	}
+
+	FGLTFMaterialDataEx MatSet;
+	MatSet.Material = const_cast<UMaterialInterface*>(Material);
+	MatSet.PropertySizes.Add(Property, OutputSize);
+	MatSet.bTangentSpaceNormal = true;
+
+	TArray<FGLTFMeshRenderData*> MeshSettings;
+	TArray<FGLTFMaterialDataEx*> MatSettings;
+	MeshSettings.Add(&MeshSet);
+	MatSettings.Add(&MatSet);
+
+	TArray<FGLTFBakeOutputEx> BakeOutputs;
+	IGLTFMaterialBakingModule& Module = FModuleManager::Get().LoadModuleChecked<IGLTFMaterialBakingModule>("GLTFMaterialBaking");
+
+	Module.SetLinearBake(true);
+	Module.BakeMaterials(MatSettings, MeshSettings, BakeOutputs);
+	const bool bIsLinearBake = Module.IsLinearBake(Property);
+	Module.SetLinearBake(false);
+
+	FGLTFBakeOutputEx& BakeOutput = BakeOutputs[0];
+
+	TGLTFSharedArray<FColor> BakedPixels = MakeShared<TArray<FColor>>(MoveTemp(BakeOutput.PropertyData.FindChecked(Property)));
+	const FIntPoint BakedSize = BakeOutput.PropertySizes.FindChecked(Property);
+	const float EmissiveScale = BakeOutput.EmissiveScale;
+
+	if (bFillAlpha)
+	{
+		// NOTE: alpha is 0 by default after baking a property, but we prefer 255 (1.0).
+		// It makes it easier to view the exported textures.
+		for (FColor& Pixel: *BakedPixels)
+		{
+			Pixel.A = 255;
+		}
+	}
+
+	if (bAdjustNormalmaps && IsNormalMap(Property))
+	{
+		// TODO: add support for adjusting normals in baking module instead
+		FGLTFTextureUtilities::FlipGreenChannel(*BakedPixels);
+	}
+
+	bool bFromSRGB = !bIsLinearBake;
+	bool bToSRGB = IsSRGB(Property);
+	FGLTFTextureUtilities::TransformColorSpace(*BakedPixels, bFromSRGB, bToSRGB);
+
+	FGLTFPropertyBakeOutput PropertyBakeOutput(Property, PF_B8G8R8A8, BakedPixels, BakedSize, EmissiveScale, !bIsLinearBake);
+
+	if (BakedPixels->Num() == 1)
+	{
+		const FColor& Pixel = (*BakedPixels)[0];
+
+		PropertyBakeOutput.bIsConstant = true;
+		PropertyBakeOutput.ConstantValue = bToSRGB ? FLinearColor(Pixel) : Pixel.ReinterpretAsLinear();
+	}
+
+	return PropertyBakeOutput;
+}
+
+FGLTFJsonTexture* FGLTFMaterialUtilities::AddTexture(FGLTFConvertBuilder& Builder, TGLTFSharedArray<FColor>& Pixels, const FIntPoint& TextureSize, bool bIgnoreAlpha, bool bIsNormalMap, const FString& TextureName, TextureAddress TextureAddress, TextureFilter TextureFilter)
+{
+	// TODO: reuse same texture index when image is the same
+	FGLTFJsonTexture* JsonTexture = Builder.AddTexture();
+	JsonTexture->Name = TextureName;
+	JsonTexture->Sampler = Builder.AddUniqueSampler(TextureAddress, TextureFilter);
+	JsonTexture->Source = Builder.AddUniqueImage(Pixels, TextureSize, bIgnoreAlpha, TextureName);
+
+	return JsonTexture;
+}
+
+FLinearColor FGLTFMaterialUtilities::GetMask(const FExpressionInput& ExpressionInput)
+{
+	return FLinearColor(ExpressionInput.MaskR, ExpressionInput.MaskG, ExpressionInput.MaskB, ExpressionInput.MaskA);
+}
+
+uint32 FGLTFMaterialUtilities::GetMaskComponentCount(const FExpressionInput& ExpressionInput)
+{
+	return ExpressionInput.MaskR + ExpressionInput.MaskG + ExpressionInput.MaskB + ExpressionInput.MaskA;
+}
+
+bool FGLTFMaterialUtilities::TryGetTextureCoordinateIndex(const UMaterialExpressionTextureSample* TextureSampler, int32& TexCoord, FGLTFJsonTextureTransform& Transform)
+{
+	const UMaterialExpression* Expression = TextureSampler->Coordinates.Expression;
+	if (Expression == nullptr)
+	{
+		TexCoord = TextureSampler->ConstCoordinate;
+		Transform = {};
+		return true;
+	}
+
+	if (const UMaterialExpressionTextureCoordinate* TextureCoordinate = Cast<UMaterialExpressionTextureCoordinate>(Expression))
+	{
+		TexCoord = TextureCoordinate->CoordinateIndex;
+		Transform.Offset.X = TextureCoordinate->UnMirrorU ? TextureCoordinate->UTiling * 0.5 : 0.0;
+		Transform.Offset.Y = TextureCoordinate->UnMirrorV ? TextureCoordinate->VTiling * 0.5 : 0.0;
+		Transform.Scale.X = TextureCoordinate->UTiling * (TextureCoordinate->UnMirrorU ? 0.5 : 1.0);
+		Transform.Scale.Y = TextureCoordinate->VTiling * (TextureCoordinate->UnMirrorV ? 0.5 : 1.0);
+		Transform.Rotation = 0;
+		return true;
+	}
+
+	// TODO: add support for advanced expression tree (ex UMaterialExpressionTextureCoordinate -> UMaterialExpressionMultiply -> UMaterialExpressionAdd)
+
+	return false;
+}
+
+void FGLTFMaterialUtilities::GetAllTextureCoordinateIndices(const UMaterialInterface* InMaterial, const FGLTFMaterialPropertyEx& InProperty, FGLTFIndexArray& OutTexCoords)
+{
+	FGLTFMaterialAnalysis Analysis;
+	AnalyzeMaterialProperty(InMaterial, InProperty, Analysis);
+
+	const TBitArray<>& TexCoords = Analysis.TextureCoordinates;
+	for (int32 Index = 0; Index < TexCoords.Num(); Index++)
+	{
+		if (TexCoords[Index])
+		{
+			OutTexCoords.Add(Index);
+		}
+	}
+}
+
+void FGLTFMaterialUtilities::AnalyzeMaterialProperty(const UMaterialInterface* InMaterial, const FGLTFMaterialPropertyEx& InProperty, FGLTFMaterialAnalysis& OutAnalysis)
+{
+	if (GetInputForProperty(InMaterial, InProperty) == nullptr)
+	{
+		OutAnalysis = FGLTFMaterialAnalysis();
+		return;
+	}
+
+	UGLTFMaterialAnalyzer::AnalyzeMaterialPropertyEx(InMaterial, InProperty.Type, InProperty.CustomOutput.ToString(), OutAnalysis);
+}
+
+FMaterialShadingModelField FGLTFMaterialUtilities::EvaluateShadingModelExpression(const UMaterialInterface* Material)
+{
+	FGLTFMaterialAnalysis Analysis;
+	AnalyzeMaterialProperty(Material, MP_ShadingModel, Analysis);
+
+	int32 Value;
+	if (FDefaultValueHelper::ParseInt(Analysis.ParameterCode, Value))
+	{
+		return static_cast<EMaterialShadingModel>(Value);
+	}
+
+	return Analysis.ShadingModels;
+}
+
+#endif
+
+EMaterialShadingModel FGLTFMaterialUtilities::GetRichestShadingModel(const FMaterialShadingModelField& ShadingModels)
+{
+	if (ShadingModels.HasShadingModel(MSM_ClearCoat))
+	{
+		return MSM_ClearCoat;
+	}
+
+	if (ShadingModels.HasShadingModel(MSM_DefaultLit))
+	{
+		return MSM_DefaultLit;
+	}
+
+	if (ShadingModels.HasShadingModel(MSM_Unlit))
+	{
+		return MSM_Unlit;
+	}
+
+	// TODO: add more shading models when conversion supported
+
+	return ShadingModels.GetFirstShadingModel();
+}
+
+FString FGLTFMaterialUtilities::ShadingModelsToString(const FMaterialShadingModelField& ShadingModels)
+{
+	FString Result;
+
+	for (uint32 Index = 0; Index < MSM_NUM; Index++)
+	{
+		const EMaterialShadingModel ShadingModel = static_cast<EMaterialShadingModel>(Index);
+		if (ShadingModels.HasShadingModel(ShadingModel))
+		{
+			FString Name = FGLTFNameUtilities::GetName(ShadingModel);
+			Result += Result.IsEmpty() ? Name : TEXT(", ") + Name;
+		}
+	}
+
+	return Result;
+}
+
+bool FGLTFMaterialUtilities::NeedsMeshData(const UMaterialInterface* Material)
+{
+#if WITH_EDITOR
+	if (Material != nullptr && !FGLTFProxyMaterialUtilities::IsProxyMaterial(Material))
+	{
+		// TODO: only analyze properties that will be needed for this specific material
+		const TArray<FGLTFMaterialPropertyEx> Properties =
+		{
+			MP_BaseColor,
+			MP_EmissiveColor,
+			MP_Opacity,
+			MP_OpacityMask,
+			MP_Metallic,
+			MP_Roughness,
+			MP_Normal,
+			MP_AmbientOcclusion,
+			MP_CustomData0,
+			MP_CustomData1,
+			FGLTFMaterialPropertyEx::ClearCoatBottomNormal,
+			// TODO: add TransmittanceColor when supported
+			// TODO: add Refraction when supported
+			// TODO: add Specular when supported
+		};
+
+		bool bNeedsMeshData = false;
+		FGLTFMaterialAnalysis Analysis;
+
+		// TODO: optimize baking by separating need for vertex data and primitive data
+
+		for (const FGLTFMaterialPropertyEx& Property: Properties)
+		{
+			AnalyzeMaterialProperty(Material, Property, Analysis);
+			bNeedsMeshData |= Analysis.bRequiresVertexData;
+			bNeedsMeshData |= Analysis.bRequiresPrimitiveData;
+		}
+
+		return bNeedsMeshData;
+	}
+#endif
+
+	return false;
+}
+
+bool FGLTFMaterialUtilities::NeedsMeshData(const TArray<const UMaterialInterface*>& Materials)
+{
+#if WITH_EDITOR
+	for (const UMaterialInterface* Material: Materials)
+	{
+		if (NeedsMeshData(Material))
+		{
+			return true;
+		}
+	}
+#endif
+
+	return false;
+}
diff --git a/Engine/Plugins/Enterprise/GLTFExporter/Source/GLTFExporter/Private/Converters/GLTFMaterialUtilities.h b/Engine/Plugins/Enterprise/GLTFExporter/Source/GLTFExporter/Private/Converters/GLTFMaterialUtilities.h
new file mode 100644
--- /dev/null
+++ b/Engine/Plugins/Enterprise/GLTFExporter/Source/GLTFExporter/Private/Converters/GLTFMaterialUtilities.h
@@ -0,0 +1,79 @@
+// Copyright Epic Games, Inc. All Rights Reserved.
+
+#pragma once
+
+#include "Builders/GLTFContainerBuilder.h"
+#include "Converters/GLTFSharedArray.h"
+#include "MaterialExpressionIO.h"
+
+struct FGLTFMaterialAnalysis;
+struct FGLTFMaterialPropertyEx;
+class UMaterialExpressionTextureSample;
+class UMaterialInstanceConstant;
+class UMaterialInterface;
+class UMaterialExpressionCustomOutput;
+
+struct FGLTFPropertyBakeOutput
+{
+	FGLTFPropertyBakeOutput(const FGLTFMaterialPropertyEx& Property, EPixelFormat PixelFormat, const TGLTFSharedArray<FColor>& Pixels, FIntPoint Size, float EmissiveScale, bool bIsSRGB)
+		: Property(Property), PixelFormat(PixelFormat), Pixels(Pixels), Size(Size), EmissiveScale(EmissiveScale), bIsSRGB(bIsSRGB), bIsConstant(false)
+	{}
+
+	const FGLTFMaterialPropertyEx& Property;
+	EPixelFormat PixelFormat;
+	TGLTFSharedArray<FColor> Pixels;
+	FIntPoint Size;
+	float EmissiveScale;
+	bool bIsSRGB;
+	bool bIsConstant;
+	FLinearColor ConstantValue;
+};
+
+struct FGLTFMaterialUtilities
+{
+	static UMaterialInterface* GetDefaultMaterial();
+
+	static bool IsClearCoatBottomNormalEnabled();
+
+#if WITH_EDITOR
+	static bool IsNormalMap(const FGLTFMaterialPropertyEx& Property);
+	static bool IsSRGB(const FGLTFMaterialPropertyEx& Property);
+
+	static FGuid GetAttributeID(const FGLTFMaterialPropertyEx& Property);
+	static FGuid GetAttributeIDChecked(const FGLTFMaterialPropertyEx& Property);
+
+	static FVector4f GetPropertyDefaultValue(const FGLTFMaterialPropertyEx& Property);
+	static FVector4f GetPropertyMask(const FGLTFMaterialPropertyEx& Property);
+
+	static const FExpressionInput* GetInputForProperty(const UMaterialInterface* Material, const FGLTFMaterialPropertyEx& Property);
+
+	template <class InputType>
+	static const FMaterialInput<InputType>* GetInputForProperty(const UMaterialInterface* Material, const FGLTFMaterialPropertyEx& Property)
+	{
+		const FExpressionInput* ExpressionInput = GetInputForProperty(Material, Property);
+		return static_cast<const FMaterialInput<InputType>*>(ExpressionInput);
+	}
+
+	static const UMaterialExpressionCustomOutput* GetCustomOutputByName(const UMaterialInterface* Material, const FString& FunctionName);
+
+	static FGLTFPropertyBakeOutput BakeMaterialProperty(const FIntPoint& OutputSize, const FGLTFMaterialPropertyEx& Property, const UMaterialInterface* Material, int32 TexCoord, const FGLTFMeshData* MeshData = nullptr, const FGLTFIndexArray& MeshSectionIndices = {}, bool bFillAlpha = true, bool bAdjustNormalmaps = true);
+
+	static FGLTFJsonTexture* AddTexture(FGLTFConvertBuilder& Builder, TGLTFSharedArray<FColor>& Pixels, const FIntPoint& TextureSize, bool bIgnoreAlpha, bool bIsNormalMap, const FString& TextureName, TextureAddress TextureAddress, TextureFilter TextureFilter);
+
+	static FLinearColor GetMask(const FExpressionInput& ExpressionInput);
+	static uint32 GetMaskComponentCount(const FExpressionInput& ExpressionInput);
+
+	static bool TryGetTextureCoordinateIndex(const UMaterialExpressionTextureSample* TextureSampler, int32& TexCoord, FGLTFJsonTextureTransform& Transform);
+	static void GetAllTextureCoordinateIndices(const UMaterialInterface* InMaterial, const FGLTFMaterialPropertyEx& InProperty, FGLTFIndexArray& OutTexCoords);
+
+	static void AnalyzeMaterialProperty(const UMaterialInterface* Material, const FGLTFMaterialPropertyEx& InProperty, FGLTFMaterialAnalysis& OutAnalysis);
+
+	static FMaterialShadingModelField EvaluateShadingModelExpression(const UMaterialInterface* Material);
+#endif
+
+	static EMaterialShadingModel GetRichestShadingModel(const FMaterialShadingModelField& ShadingModels);
+	static FString ShadingModelsToString(const FMaterialShadingModelField& ShadingModels);
+
+	static bool NeedsMeshData(const UMaterialInterface* Material);
+	static bool NeedsMeshData(const TArray<const UMaterialInterface*>& Materials);
+};
diff --git a/Engine/Plugins/Enterprise/GLTFExporter/Source/GLTFExporter/Private/Converters/GLTFMaterialUtility.cpp b/Engine/Plugins/Enterprise/GLTFExporter/Source/GLTFExporter/Private/Converters/GLTFMaterialUtility.cpp
deleted file mode 100644
diff --git a/Engine/Plugins/Enterprise/GLTFExporter/Source/GLTFExporter/Private/Converters/GLTFMaterialUtility.h b/Engine/Plugins/Enterprise/GLTFExporter/Source/GLTFExporter/Private/Converters/GLTFMaterialUtility.h
deleted file mode 100644
diff --git a/Engine/Plugins/Enterprise/GLTFExporter/Source/GLTFExporter/Private/Converters/GLTFMaterialVariantConverters.cpp b/Engine/Plugins/Enterprise/GLTFExporter/Source/GLTFExporter/Private/Converters/GLTFMaterialVariantConverters.cpp
--- a/Engine/Plugins/Enterprise/GLTFExporter/Source/GLTFExporter/Private/Converters/GLTFMaterialVariantConverters.cpp
+++ b/Engine/Plugins/Enterprise/GLTFExporter/Source/GLTFExporter/Private/Converters/GLTFMaterialVariantConverters.cpp
@@ -1,7 +1,7 @@
 ﻿// Copyright Epic Games, Inc. All Rights Reserved.
 
 #include "Converters/GLTFMaterialVariantConverters.h"
-#include "Converters/GLTFVariantUtility.h"
+#include "Converters/GLTFVariantUtilities.h"
 #include "Builders/GLTFContainerBuilder.h"
 #include "Components/MeshComponent.h"
 #include "VariantObjectBinding.h"
@@ -89,7 +89,7 @@ bool FGLTFMaterialVariantConverter::TryParseMaterialProperty(FGLTFJsonPrimitive*
 	{
 		Builder.LogWarning(FString::Printf(
 			TEXT("Variant property %s must belong to a mesh component, the property will be skipped"),
-			*FGLTFVariantUtility::GetLogContext(Property)));
+			*FGLTFVariantUtilities::GetLogContext(Property)));
 		return false;
 	}
 
@@ -98,7 +98,7 @@ bool FGLTFMaterialVariantConverter::TryParseMaterialProperty(FGLTFJsonPrimitive*
 	{
 		Builder.LogWarning(FString::Printf(
 			TEXT("Variant property %s must belong to an actor, the property will be skipped"),
-			*FGLTFVariantUtility::GetLogContext(Property)));
+			*FGLTFVariantUtilities::GetLogContext(Property)));
 		return false;
 	}
 
@@ -106,7 +106,7 @@ bool FGLTFMaterialVariantConverter::TryParseMaterialProperty(FGLTFJsonPrimitive*
 	{
 		Builder.LogWarning(FString::Printf(
 			TEXT("Variant property %s doesn't belong to an actor selected for export, the property will be skipped"),
-			*FGLTFVariantUtility::GetLogContext(Property)));
+			*FGLTFVariantUtilities::GetLogContext(Property)));
 		return false;
 	}
 
@@ -117,7 +117,7 @@ bool FGLTFMaterialVariantConverter::TryParseMaterialProperty(FGLTFJsonPrimitive*
 	{
 		Builder.LogWarning(FString::Printf(
 			TEXT("Failed to parse material index for variant property %s, the property will be skipped"),
-			*FGLTFVariantUtility::GetLogContext(Property)));
+			*FGLTFVariantUtilities::GetLogContext(Property)));
 		return false;
 	}
 
@@ -132,7 +132,7 @@ bool FGLTFMaterialVariantConverter::TryParseMaterialProperty(FGLTFJsonPrimitive*
 
 		Builder.LogWarning(FString::Printf(
 			TEXT("No material assigned, the property will be skipped. Context: %s"),
-			*FGLTFVariantUtility::GetLogContext(Property)));
+			*FGLTFVariantUtilities::GetLogContext(Property)));
 		return false;
 	}
 
@@ -143,7 +143,7 @@ bool FGLTFMaterialVariantConverter::TryParseMaterialProperty(FGLTFJsonPrimitive*
 	FGLTFJsonMesh* JsonMesh = Builder.AddUniqueMesh(Target);
 
 	OutPrimitive = &JsonMesh->Primitives[MaterialIndex];
-	OutMaterial = FGLTFVariantUtility::AddUniqueMaterial(Builder, Material, Target, MaterialIndex);
+	OutMaterial = FGLTFVariantUtilities::AddUniqueMaterial(Builder, Material, Target, MaterialIndex);
 
 	return true;
 }
diff --git a/Engine/Plugins/Enterprise/GLTFExporter/Source/GLTFExporter/Private/Converters/GLTFMeshConverters.cpp b/Engine/Plugins/Enterprise/GLTFExporter/Source/GLTFExporter/Private/Converters/GLTFMeshConverters.cpp
--- a/Engine/Plugins/Enterprise/GLTFExporter/Source/GLTFExporter/Private/Converters/GLTFMeshConverters.cpp
+++ b/Engine/Plugins/Enterprise/GLTFExporter/Source/GLTFExporter/Private/Converters/GLTFMeshConverters.cpp
@@ -1,8 +1,8 @@
 // Copyright Epic Games, Inc. All Rights Reserved.
 
 #include "Converters/GLTFMeshConverters.h"
-#include "Converters/GLTFMeshUtility.h"
-#include "Converters/GLTFMaterialUtility.h"
+#include "Converters/GLTFMeshUtilities.h"
+#include "Converters/GLTFMaterialUtilities.h"
 #include "Builders/GLTFConvertBuilder.h"
 #include "Tasks/GLTFDelayedMeshTasks.h"
 #include "Engine/SkeletalMesh.h"
@@ -13,7 +13,7 @@
 
 void FGLTFStaticMeshConverter::Sanitize(const UStaticMesh*& StaticMesh, const UStaticMeshComponent*& StaticMeshComponent, FGLTFMaterialArray& Materials, int32& LODIndex)
 {
-	FGLTFMeshUtility::ResolveMaterials(Materials, StaticMeshComponent, StaticMesh);
+	FGLTFMeshUtilities::ResolveMaterials(Materials, StaticMeshComponent, StaticMesh);
 	Builder.ResolveProxies(Materials);
 
 	LODIndex = Builder.SanitizeLOD(StaticMesh, StaticMeshComponent, LODIndex);
@@ -21,7 +21,7 @@ void FGLTFStaticMeshConverter::Sanitize(const UStaticMesh*& StaticMesh, const US
 	if (StaticMeshComponent != nullptr)
 	{
 		const bool bUsesMeshData = Builder.ExportOptions->BakeMaterialInputs == EGLTFMaterialBakeMode::UseMeshData &&
-			FGLTFMaterialUtility::NeedsMeshData(Materials); // TODO: if this expensive, cache the results for each material
+			FGLTFMaterialUtilities::NeedsMeshData(Materials); // TODO: if this expensive, cache the results for each material
 
 		const bool bIsReferencedByVariant = Builder.GetObjectVariants(StaticMeshComponent) != nullptr;
 
@@ -45,7 +45,7 @@ FGLTFJsonMesh* FGLTFStaticMeshConverter::Convert(const UStaticMesh* StaticMesh,
 	}
 #endif
 
-	const int32 MaterialCount = FGLTFMeshUtility::GetMaterials(StaticMesh).Num();
+	const int32 MaterialCount = FGLTFMeshUtilities::GetMaterials(StaticMesh).Num();
 	FGLTFJsonMesh* JsonMesh = Builder.AddMesh();
 	JsonMesh->Primitives.AddDefaulted(MaterialCount);
 
@@ -55,7 +55,7 @@ FGLTFJsonMesh* FGLTFStaticMeshConverter::Convert(const UStaticMesh* StaticMesh,
 
 void FGLTFSkeletalMeshConverter::Sanitize(const USkeletalMesh*& SkeletalMesh, const USkeletalMeshComponent*& SkeletalMeshComponent, FGLTFMaterialArray& Materials, int32& LODIndex)
 {
-	FGLTFMeshUtility::ResolveMaterials(Materials, SkeletalMeshComponent, SkeletalMesh);
+	FGLTFMeshUtilities::ResolveMaterials(Materials, SkeletalMeshComponent, SkeletalMesh);
 	Builder.ResolveProxies(Materials);
 
 	LODIndex = Builder.SanitizeLOD(SkeletalMesh, SkeletalMeshComponent, LODIndex);
@@ -63,7 +63,7 @@ void FGLTFSkeletalMeshConverter::Sanitize(const USkeletalMesh*& SkeletalMesh, co
 	if (SkeletalMeshComponent != nullptr)
 	{
 		const bool bUsesMeshData = Builder.ExportOptions->BakeMaterialInputs == EGLTFMaterialBakeMode::UseMeshData &&
-			FGLTFMaterialUtility::NeedsMeshData(Materials); // TODO: if this expensive, cache the results for each material
+			FGLTFMaterialUtilities::NeedsMeshData(Materials); // TODO: if this expensive, cache the results for each material
 
 		const bool bIsReferencedByVariant = Builder.GetObjectVariants(SkeletalMeshComponent) != nullptr;
 
@@ -88,7 +88,7 @@ FGLTFJsonMesh* FGLTFSkeletalMeshConverter::Convert(const USkeletalMesh* Skeletal
 	}
 #endif
 
-	const int32 MaterialCount = FGLTFMeshUtility::GetMaterials(SkeletalMesh).Num();
+	const int32 MaterialCount = FGLTFMeshUtilities::GetMaterials(SkeletalMesh).Num();
 	FGLTFJsonMesh* JsonMesh = Builder.AddMesh();
 	JsonMesh->Primitives.AddDefaulted(MaterialCount);
 
diff --git a/Engine/Plugins/Enterprise/GLTFExporter/Source/GLTFExporter/Private/Converters/GLTFMeshData.cpp b/Engine/Plugins/Enterprise/GLTFExporter/Source/GLTFExporter/Private/Converters/GLTFMeshData.cpp
--- a/Engine/Plugins/Enterprise/GLTFExporter/Source/GLTFExporter/Private/Converters/GLTFMeshData.cpp
+++ b/Engine/Plugins/Enterprise/GLTFExporter/Source/GLTFExporter/Private/Converters/GLTFMeshData.cpp
@@ -1,7 +1,7 @@
 // Copyright Epic Games, Inc. All Rights Reserved.
 
 #include "Converters/GLTFMeshData.h"
-#include "Converters/GLTFNameUtility.h"
+#include "Converters/GLTFNameUtilities.h"
 
 #include "Components/SkeletalMeshComponent.h"
 #include "Components/StaticMeshComponent.h"
@@ -31,7 +31,7 @@ FGLTFMeshData::FGLTFMeshData(const UStaticMesh* StaticMesh, const UStaticMeshCom
 
 	if (StaticMeshComponent != nullptr)
 	{
-		Name = FGLTFNameUtility::GetName(StaticMeshComponent);
+		Name = FGLTFNameUtilities::GetName(StaticMeshComponent);
 
 #if WITH_EDITOR
 		PrimitiveData = { StaticMeshComponent };
@@ -83,7 +83,7 @@ FGLTFMeshData::FGLTFMeshData(const USkeletalMesh* SkeletalMesh, const USkeletalM
 
 	if (SkeletalMeshComponent != nullptr)
 	{
-		Name = FGLTFNameUtility::GetName(SkeletalMeshComponent);
+		Name = FGLTFNameUtilities::GetName(SkeletalMeshComponent);
 
 #if WITH_EDITOR
 		PrimitiveData = { SkeletalMeshComponent };
diff --git a/Engine/Plugins/Enterprise/GLTFExporter/Source/GLTFExporter/Private/Converters/GLTFMeshUtilities.cpp b/Engine/Plugins/Enterprise/GLTFExporter/Source/GLTFExporter/Private/Converters/GLTFMeshUtilities.cpp
new file mode 100644
--- /dev/null
+++ b/Engine/Plugins/Enterprise/GLTFExporter/Source/GLTFExporter/Private/Converters/GLTFMeshUtilities.cpp
@@ -0,0 +1,246 @@
+// Copyright Epic Games, Inc. All Rights Reserved.
+
+#include "Converters/GLTFMeshUtilities.h"
+#include "Components/SkeletalMeshComponent.h"
+#include "Engine/StaticMesh.h"
+#include "Engine/SkeletalMesh.h"
+#include "Engine/SkinnedAssetCommon.h"
+#include "MaterialDomain.h"
+#include "Materials/Material.h"
+#include "Rendering/SkeletalMeshRenderData.h"
+#include "StaticMeshCompiler.h"
+#include "SkinnedAssetCompiler.h"
+
+void FGLTFMeshUtilities::FullyLoad(const UStaticMesh* InStaticMesh)
+{
+	UStaticMesh* StaticMesh = const_cast<UStaticMesh*>(InStaticMesh);
+
+#if WITH_EDITOR
+	FStaticMeshCompilingManager::Get().FinishCompilation({ StaticMesh });
+#endif
+
+	StaticMesh->SetForceMipLevelsToBeResident(30.0f);
+	StaticMesh->WaitForStreaming();
+}
+
+void FGLTFMeshUtilities::FullyLoad(const USkeletalMesh* InSkeletalMesh)
+{
+	USkeletalMesh* SkeletalMesh = const_cast<USkeletalMesh*>(InSkeletalMesh);
+
+#if WITH_EDITOR
+	FSkinnedAssetCompilingManager::Get().FinishCompilation({ SkeletalMesh });
+#endif
+
+	SkeletalMesh->SetForceMipLevelsToBeResident(30.0f);
+	SkeletalMesh->WaitForStreaming();
+}
+
+const UMaterialInterface* FGLTFMeshUtilities::GetDefaultMaterial()
+{
+	static UMaterial* DefaultMaterial = UMaterial::GetDefaultMaterial(MD_Surface);
+	return DefaultMaterial;
+}
+
+const TArray<FStaticMaterial>& FGLTFMeshUtilities::GetMaterials(const UStaticMesh* StaticMesh)
+{
+	return StaticMesh->GetStaticMaterials();
+}
+
+const TArray<FSkeletalMaterial>& FGLTFMeshUtilities::GetMaterials(const USkeletalMesh* SkeletalMesh)
+{
+	return SkeletalMesh->GetMaterials();
+}
+
+const UMaterialInterface* FGLTFMeshUtilities::GetMaterial(const UMaterialInterface* Material)
+{
+	return Material;
+}
+
+const UMaterialInterface* FGLTFMeshUtilities::GetMaterial(const FStaticMaterial& Material)
+{
+	return Material.MaterialInterface;
+}
+
+const UMaterialInterface* FGLTFMeshUtilities::GetMaterial(const FSkeletalMaterial& Material)
+{
+	return Material.MaterialInterface;
+}
+
+void FGLTFMeshUtilities::ResolveMaterials(TArray<const UMaterialInterface*>& Materials, const UStaticMeshComponent* StaticMeshComponent, const UStaticMesh* StaticMesh)
+{
+	if (StaticMeshComponent != nullptr)
+	{
+		ResolveMaterials(Materials, StaticMeshComponent->GetMaterials());
+	}
+
+	if (StaticMesh != nullptr)
+	{
+		ResolveMaterials(Materials, GetMaterials(StaticMesh));
+	}
+
+	ResolveMaterials(Materials, GetDefaultMaterial());
+}
+
+void FGLTFMeshUtilities::ResolveMaterials(TArray<const UMaterialInterface*>& Materials, const USkeletalMeshComponent* SkeletalMeshComponent, const USkeletalMesh* SkeletalMesh)
+{
+	if (SkeletalMeshComponent != nullptr)
+	{
+		ResolveMaterials(Materials, SkeletalMeshComponent->GetMaterials());
+	}
+
+	if (SkeletalMesh != nullptr)
+	{
+		ResolveMaterials(Materials, GetMaterials(SkeletalMesh));
+	}
+
+	ResolveMaterials(Materials, GetDefaultMaterial());
+}
+
+template <typename MaterialType>
+void FGLTFMeshUtilities::ResolveMaterials(TArray<const UMaterialInterface*>& Materials, const TArray<MaterialType>& Defaults)
+{
+	const int32 Count = Defaults.Num();
+	Materials.SetNumZeroed(Count);
+
+	for (int32 Index = 0; Index < Count; ++Index)
+	{
+		const UMaterialInterface*& Material = Materials[Index];
+		if (Material == nullptr)
+		{
+			Material = GetMaterial(Defaults[Index]);
+		}
+	}
+}
+
+void FGLTFMeshUtilities::ResolveMaterials(TArray<const UMaterialInterface*>& Materials, const UMaterialInterface* Default)
+{
+	for (const UMaterialInterface*& Material : Materials)
+	{
+		if (Material == nullptr)
+		{
+			Material = Default;
+		}
+	}
+}
+
+FGLTFIndexArray FGLTFMeshUtilities::GetSectionIndices(const UStaticMesh* StaticMesh, int32 LODIndex, int32 MaterialIndex)
+{
+	if (StaticMesh == nullptr)
+	{
+		return {};
+	}
+
+	const FStaticMeshLODResources& MeshLOD = StaticMesh->GetLODForExport(LODIndex);
+	return GetSectionIndices(MeshLOD, MaterialIndex);
+}
+
+FGLTFIndexArray FGLTFMeshUtilities::GetSectionIndices(const USkeletalMesh* SkeletalMesh, int32 LODIndex, int32 MaterialIndex)
+{
+	if (SkeletalMesh == nullptr)
+	{
+		return {};
+	}
+
+	const FSkeletalMeshRenderData* RenderData = SkeletalMesh->GetResourceForRendering();
+	const FSkeletalMeshLODRenderData& MeshLOD = RenderData->LODRenderData[LODIndex];
+	return GetSectionIndices(MeshLOD, MaterialIndex);
+}
+
+FGLTFIndexArray FGLTFMeshUtilities::GetSectionIndices(const FStaticMeshLODResources& MeshLOD, int32 MaterialIndex)
+{
+	const FStaticMeshSectionArray& Sections = MeshLOD.Sections;
+
+	FGLTFIndexArray SectionIndices;
+	SectionIndices.Reserve(Sections.Num());
+
+	for (int32 SectionIndex = 0; SectionIndex < Sections.Num(); ++SectionIndex)
+	{
+		if (Sections[SectionIndex].MaterialIndex == MaterialIndex)
+		{
+			SectionIndices.Add(SectionIndex);
+		}
+	}
+
+	return SectionIndices;
+}
+
+FGLTFIndexArray FGLTFMeshUtilities::GetSectionIndices(const FSkeletalMeshLODRenderData& MeshLOD, int32 MaterialIndex)
+{
+	const TArray<FSkelMeshRenderSection>& Sections = MeshLOD.RenderSections;
+
+	FGLTFIndexArray SectionIndices;
+	SectionIndices.Reserve(Sections.Num());
+
+	for (int32 SectionIndex = 0; SectionIndex < Sections.Num(); ++SectionIndex)
+	{
+		if (Sections[SectionIndex].MaterialIndex == MaterialIndex)
+		{
+			SectionIndices.Add(SectionIndex);
+		}
+	}
+
+	return SectionIndices;
+}
+
+int32 FGLTFMeshUtilities::GetLOD(const UStaticMesh* StaticMesh, const UStaticMeshComponent* StaticMeshComponent, int32 DefaultLOD)
+{
+	const int32 ForcedLOD = StaticMeshComponent != nullptr ? StaticMeshComponent->ForcedLodModel - 1 : -1;
+	const int32 LOD = ForcedLOD > 0 ? ForcedLOD : FMath::Max(DefaultLOD, GetMinimumLOD(StaticMesh, StaticMeshComponent));
+	return FMath::Min(LOD, GetMaximumLOD(StaticMesh));
+}
+
+int32 FGLTFMeshUtilities::GetLOD(const USkeletalMesh* SkeletalMesh, const USkeletalMeshComponent* SkeletalMeshComponent, int32 DefaultLOD)
+{
+	const int32 ForcedLOD = SkeletalMeshComponent != nullptr ? SkeletalMeshComponent->GetForcedLOD() - 1 : -1;
+	const int32 LOD = ForcedLOD > 0 ? ForcedLOD : FMath::Max(DefaultLOD, GetMinimumLOD(SkeletalMesh, SkeletalMeshComponent));
+	return FMath::Min(LOD, GetMaximumLOD(SkeletalMesh));
+}
+
+int32 FGLTFMeshUtilities::GetMaximumLOD(const UStaticMesh* StaticMesh)
+{
+	return StaticMesh != nullptr ? StaticMesh->GetNumLODs() - 1 : -1;
+}
+
+int32 FGLTFMeshUtilities::GetMaximumLOD(const USkeletalMesh* SkeletalMesh)
+{
+	if (SkeletalMesh != nullptr)
+	{
+		const FSkeletalMeshRenderData* RenderData = SkeletalMesh->GetResourceForRendering();
+		if (RenderData != nullptr)
+		{
+			return RenderData->LODRenderData.Num() - 1;
+		}
+	}
+
+	return -1;
+}
+
+int32 FGLTFMeshUtilities::GetMinimumLOD(const UStaticMesh* StaticMesh, const UStaticMeshComponent* StaticMeshComponent)
+{
+	if (StaticMeshComponent != nullptr && StaticMeshComponent->bOverrideMinLOD)
+	{
+		return StaticMeshComponent->MinLOD;
+	}
+
+	if (StaticMesh != nullptr)
+	{
+		return StaticMesh->GetMinLOD().Default;
+	}
+
+	return -1;
+}
+
+int32 FGLTFMeshUtilities::GetMinimumLOD(const USkeletalMesh* SkeletalMesh, const USkeletalMeshComponent* SkeletalMeshComponent)
+{
+	if (SkeletalMeshComponent != nullptr && SkeletalMeshComponent->bOverrideMinLod)
+	{
+		return SkeletalMeshComponent->MinLodModel;
+	}
+
+	if (SkeletalMesh != nullptr)
+	{
+		return SkeletalMesh->GetMinLod().Default;
+	}
+
+	return -1;
+}
diff --git a/Engine/Plugins/Enterprise/GLTFExporter/Source/GLTFExporter/Private/Converters/GLTFMeshUtilities.h b/Engine/Plugins/Enterprise/GLTFExporter/Source/GLTFExporter/Private/Converters/GLTFMeshUtilities.h
new file mode 100644
--- /dev/null
+++ b/Engine/Plugins/Enterprise/GLTFExporter/Source/GLTFExporter/Private/Converters/GLTFMeshUtilities.h
@@ -0,0 +1,52 @@
+// Copyright Epic Games, Inc. All Rights Reserved.
+
+#pragma once
+
+#include "Converters/GLTFIndexArray.h"
+
+class UMaterialInterface;
+class UStaticMesh;
+class UStaticMeshComponent;
+class USkeletalMesh;
+class USkeletalMeshComponent;
+class FSkeletalMeshLODRenderData;
+struct FStaticMaterial;
+struct FSkeletalMaterial;
+struct FStaticMeshLODResources;
+
+struct FGLTFMeshUtilities
+{
+	static void FullyLoad(const UStaticMesh* InStaticMesh);
+	static void FullyLoad(const USkeletalMesh* InSkeletalMesh);
+
+	static const UMaterialInterface* GetDefaultMaterial();
+
+	static const TArray<FStaticMaterial>& GetMaterials(const UStaticMesh* StaticMesh);
+	static const TArray<FSkeletalMaterial>& GetMaterials(const USkeletalMesh* SkeletalMesh);
+
+	static const UMaterialInterface* GetMaterial(const UMaterialInterface* Material);
+	static const UMaterialInterface* GetMaterial(const FStaticMaterial& Material);
+	static const UMaterialInterface* GetMaterial(const FSkeletalMaterial& Material);
+
+	static void ResolveMaterials(TArray<const UMaterialInterface*>& Materials, const UStaticMeshComponent* StaticMeshComponent, const UStaticMesh* StaticMesh);
+	static void ResolveMaterials(TArray<const UMaterialInterface*>& Materials, const USkeletalMeshComponent* SkeletalMeshComponent, const USkeletalMesh* SkeletalMesh);
+
+	template <typename MaterialType>
+	static void ResolveMaterials(TArray<const UMaterialInterface*>& Materials, const TArray<MaterialType>& Defaults);
+	static void ResolveMaterials(TArray<const UMaterialInterface*>& Materials, const UMaterialInterface* Default);
+
+	static FGLTFIndexArray GetSectionIndices(const UStaticMesh* StaticMesh, int32 LODIndex, int32 MaterialIndex);
+	static FGLTFIndexArray GetSectionIndices(const USkeletalMesh* SkeletalMesh, int32 LODIndex, int32 MaterialIndex);
+
+	static FGLTFIndexArray GetSectionIndices(const FStaticMeshLODResources& MeshLOD, int32 MaterialIndex);
+	static FGLTFIndexArray GetSectionIndices(const FSkeletalMeshLODRenderData& MeshLOD, int32 MaterialIndex);
+
+	static int32 GetLOD(const UStaticMesh* StaticMesh, const UStaticMeshComponent* StaticMeshComponent, int32 DefaultLOD);
+	static int32 GetLOD(const USkeletalMesh* SkeletalMesh, const USkeletalMeshComponent* SkeletalMeshComponent, int32 DefaultLOD);
+
+	static int32 GetMaximumLOD(const UStaticMesh* StaticMesh);
+	static int32 GetMaximumLOD(const USkeletalMesh* SkeletalMesh);
+
+	static int32 GetMinimumLOD(const UStaticMesh* StaticMesh, const UStaticMeshComponent* StaticMeshComponent);
+	static int32 GetMinimumLOD(const USkeletalMesh* SkeletalMesh, const USkeletalMeshComponent* SkeletalMeshComponent);
+};
diff --git a/Engine/Plugins/Enterprise/GLTFExporter/Source/GLTFExporter/Private/Converters/GLTFMeshUtility.cpp b/Engine/Plugins/Enterprise/GLTFExporter/Source/GLTFExporter/Private/Converters/GLTFMeshUtility.cpp
deleted file mode 100644
diff --git a/Engine/Plugins/Enterprise/GLTFExporter/Source/GLTFExporter/Private/Converters/GLTFMeshUtility.h b/Engine/Plugins/Enterprise/GLTFExporter/Source/GLTFExporter/Private/Converters/GLTFMeshUtility.h
deleted file mode 100644
diff --git a/Engine/Plugins/Enterprise/GLTFExporter/Source/GLTFExporter/Private/Converters/GLTFNameUtilities.cpp b/Engine/Plugins/Enterprise/GLTFExporter/Source/GLTFExporter/Private/Converters/GLTFNameUtilities.cpp
new file mode 100644
--- /dev/null
+++ b/Engine/Plugins/Enterprise/GLTFExporter/Source/GLTFExporter/Private/Converters/GLTFNameUtilities.cpp
@@ -0,0 +1,29 @@
+// Copyright Epic Games, Inc. All Rights Reserved.
+
+#include "Converters/GLTFNameUtilities.h"
+#include "Components/SceneComponent.h"
+#include "GameFramework/Actor.h"
+
+FString FGLTFNameUtilities::GetName(const UEnum* Enum, int32 Value)
+{
+	check(Enum != nullptr);
+	const FString DisplayName = Enum->GetDisplayNameTextByValue(Value).ToString();
+	return DisplayName.IsEmpty() ? FString::FromInt(Value) : DisplayName;
+}
+
+FString FGLTFNameUtilities::GetName(const USceneComponent* Component)
+{
+	if (const AActor* Owner = Component->GetOwner())
+	{
+		if (Component == Owner->GetRootComponent())
+		{
+#if WITH_EDITOR
+			return Owner->GetActorLabel();
+#else
+			return Owner->GetName();
+#endif
+		}
+	}
+
+	return Component->GetName();
+}
diff --git a/Engine/Plugins/Enterprise/GLTFExporter/Source/GLTFExporter/Private/Converters/GLTFNameUtilities.h b/Engine/Plugins/Enterprise/GLTFExporter/Source/GLTFExporter/Private/Converters/GLTFNameUtilities.h
new file mode 100644
--- /dev/null
+++ b/Engine/Plugins/Enterprise/GLTFExporter/Source/GLTFExporter/Private/Converters/GLTFNameUtilities.h
@@ -0,0 +1,21 @@
+// Copyright Epic Games, Inc. All Rights Reserved.
+
+#pragma once
+
+#include "CoreMinimal.h"
+#include "UObject/ReflectedTypeAccessors.h"
+
+class USceneComponent;
+
+struct FGLTFNameUtilities
+{
+	template <typename EnumType, typename = typename TEnableIf<TIsEnum<EnumType>::Value>::Type>
+	static FString GetName(EnumType Value)
+	{
+		return GetName(StaticEnum<EnumType>(), static_cast<int32>(Value));
+	}
+
+	static FString GetName(const UEnum* Enum, int32 Value);
+
+	static FString GetName(const USceneComponent* Component);
+};
diff --git a/Engine/Plugins/Enterprise/GLTFExporter/Source/GLTFExporter/Private/Converters/GLTFNameUtility.cpp b/Engine/Plugins/Enterprise/GLTFExporter/Source/GLTFExporter/Private/Converters/GLTFNameUtility.cpp
deleted file mode 100644
diff --git a/Engine/Plugins/Enterprise/GLTFExporter/Source/GLTFExporter/Private/Converters/GLTFNameUtility.h b/Engine/Plugins/Enterprise/GLTFExporter/Source/GLTFExporter/Private/Converters/GLTFNameUtility.h
deleted file mode 100644
diff --git a/Engine/Plugins/Enterprise/GLTFExporter/Source/GLTFExporter/Private/Converters/GLTFNodeConverters.cpp b/Engine/Plugins/Enterprise/GLTFExporter/Source/GLTFExporter/Private/Converters/GLTFNodeConverters.cpp
--- a/Engine/Plugins/Enterprise/GLTFExporter/Source/GLTFExporter/Private/Converters/GLTFNodeConverters.cpp
+++ b/Engine/Plugins/Enterprise/GLTFExporter/Source/GLTFExporter/Private/Converters/GLTFNodeConverters.cpp
@@ -3,7 +3,7 @@
 #include "Converters/GLTFNodeConverters.h"
 #include "Builders/GLTFContainerBuilder.h"
 #include "Utilities/GLTFCoreUtilities.h"
-#include "Converters/GLTFNameUtility.h"
+#include "Converters/GLTFNameUtilities.h"
 #include "LevelSequenceActor.h"
 #include "Camera/CameraComponent.h"
 #include "Components/LightComponent.h"
@@ -102,7 +102,7 @@ FGLTFJsonNode* FGLTFComponentConverter::Convert(const USceneComponent* SceneComp
 	const FTransform3f RelativeTransform = bIsRootNode ? Transform : Transform.GetRelativeTransform(ParentTransform);
 
 	FGLTFJsonNode* Node = Builder.AddNode();
-	Node->Name = FGLTFNameUtility::GetName(SceneComponent);
+	Node->Name = FGLTFNameUtilities::GetName(SceneComponent);
 	Node->Translation = FGLTFCoreUtilities::ConvertPosition(RelativeTransform.GetTranslation(), Builder.ExportOptions->ExportUniformScale);
 	Node->Rotation = FGLTFCoreUtilities::ConvertRotation(RelativeTransform.GetRotation());
 	Node->Scale = FGLTFCoreUtilities::ConvertScale(RelativeTransform.GetScale3D());
@@ -156,7 +156,7 @@ void FGLTFComponentConverter::ConvertComponentSpecialization(const USceneCompone
 		{
 			// TODO: conversion of camera direction should be done in separate converter
 			FGLTFJsonNode* CameraNode = Builder.AddNode();
-			CameraNode->Name = FGLTFNameUtility::GetName(CameraComponent);
+			CameraNode->Name = FGLTFNameUtilities::GetName(CameraComponent);
 			CameraNode->Rotation = FGLTFCoreUtilities::GetLocalCameraRotation();
 			CameraNode->Camera = Builder.AddUniqueCamera(CameraComponent);
 			Node->Children.Add(CameraNode);
@@ -168,7 +168,7 @@ void FGLTFComponentConverter::ConvertComponentSpecialization(const USceneCompone
 		{
 			// TODO: conversion of light direction should be done in separate converter
 			FGLTFJsonNode* LightNode = Builder.AddNode();
-			LightNode->Name = FGLTFNameUtility::GetName(LightComponent);
+			LightNode->Name = FGLTFNameUtilities::GetName(LightComponent);
 			LightNode->Rotation = FGLTFCoreUtilities::GetLocalLightRotation();
 			LightNode->Light = Builder.AddUniqueLight(LightComponent);
 			Node->Children.Add(LightNode);
diff --git a/Engine/Plugins/Enterprise/GLTFExporter/Source/GLTFExporter/Private/Converters/GLTFSamplerConverters.cpp b/Engine/Plugins/Enterprise/GLTFExporter/Source/GLTFExporter/Private/Converters/GLTFSamplerConverters.cpp
--- a/Engine/Plugins/Enterprise/GLTFExporter/Source/GLTFExporter/Private/Converters/GLTFSamplerConverters.cpp
+++ b/Engine/Plugins/Enterprise/GLTFExporter/Source/GLTFExporter/Private/Converters/GLTFSamplerConverters.cpp
@@ -3,13 +3,13 @@
 #include "Converters/GLTFSamplerConverters.h"
 #include "Builders/GLTFContainerBuilder.h"
 #include "Utilities/GLTFCoreUtilities.h"
-#include "Converters/GLTFTextureUtility.h"
+#include "Converters/GLTFTextureUtilities.h"
 
 void FGLTFSamplerConverter::Sanitize(TextureAddress& AddressX, TextureAddress& AddressY, TextureFilter& Filter, TextureGroup& LODGroup)
 {
 	if (Filter == TF_Default)
 	{
-		Filter = FGLTFTextureUtility::GetDefaultFilter(LODGroup);
+		Filter = FGLTFTextureUtilities::GetDefaultFilter(LODGroup);
 	}
 
 	LODGroup = TEXTUREGROUP_MAX; // Ignore it, since Filter should cover any missing information
diff --git a/Engine/Plugins/Enterprise/GLTFExporter/Source/GLTFExporter/Private/Converters/GLTFSkinConverters.cpp b/Engine/Plugins/Enterprise/GLTFExporter/Source/GLTFExporter/Private/Converters/GLTFSkinConverters.cpp
--- a/Engine/Plugins/Enterprise/GLTFExporter/Source/GLTFExporter/Private/Converters/GLTFSkinConverters.cpp
+++ b/Engine/Plugins/Enterprise/GLTFExporter/Source/GLTFExporter/Private/Converters/GLTFSkinConverters.cpp
@@ -2,7 +2,7 @@
 
 #include "Converters/GLTFSkinConverters.h"
 #include "Utilities/GLTFCoreUtilities.h"
-#include "Converters/GLTFBoneUtility.h"
+#include "Converters/GLTFBoneUtilities.h"
 #include "Builders/GLTFConvertBuilder.h"
 #include "Animation/Skeleton.h"
 #include "Engine/SkeletalMesh.h"
@@ -34,7 +34,7 @@ FGLTFJsonSkin* FGLTFSkinConverter::Convert(FGLTFJsonNode* RootNode, const USkele
 
 	for (int32 BoneIndex = 0; BoneIndex < BoneCount; ++BoneIndex)
 	{
-		const FTransform3f InverseBindTransform = FTransform3f(FGLTFBoneUtility::GetBindTransform(RefSkeleton, BoneIndex).Inverse());
+		const FTransform3f InverseBindTransform = FTransform3f(FGLTFBoneUtilities::GetBindTransform(RefSkeleton, BoneIndex).Inverse());
 		InverseBindMatrices[BoneIndex] = FGLTFCoreUtilities::ConvertTransform(InverseBindTransform, Builder.ExportOptions->ExportUniformScale);
 	}
 
diff --git a/Engine/Plugins/Enterprise/GLTFExporter/Source/GLTFExporter/Private/Converters/GLTFTextureUtilities.cpp b/Engine/Plugins/Enterprise/GLTFExporter/Source/GLTFExporter/Private/Converters/GLTFTextureUtilities.cpp
new file mode 100644
--- /dev/null
+++ b/Engine/Plugins/Enterprise/GLTFExporter/Source/GLTFExporter/Private/Converters/GLTFTextureUtilities.cpp
@@ -0,0 +1,200 @@
+// Copyright Epic Games, Inc. All Rights Reserved.
+
+#include "Converters/GLTFTextureUtilities.h"
+#include "Converters/GLTFNormalMapPreview.h"
+#include "Converters/GLTFSimpleTexture2DPreview.h"
+#include "CanvasItem.h"
+#include "CanvasTypes.h"
+#include "TextureCompiler.h"
+#include "Engine/Texture2D.h"
+#include "Engine/TextureCube.h"
+#include "Engine/TextureRenderTarget2D.h"
+#include "Engine/TextureRenderTargetCube.h"
+#include "DeviceProfiles/DeviceProfile.h"
+#include "DeviceProfiles/DeviceProfileManager.h"
+#include "TextureResource.h"
+
+bool FGLTFTextureUtilities::IsAlphaless(EPixelFormat PixelFormat)
+{
+	switch (PixelFormat)
+	{
+		case PF_ATC_RGB:
+		case PF_BC4:
+		case PF_BC5:
+		case PF_DXT1:
+		case PF_ETC1:
+		case PF_ETC2_RGB:
+		case PF_FloatR11G11B10:
+		case PF_FloatRGB:
+		case PF_R5G6B5_UNORM:
+			// TODO: add more pixel formats that don't support alpha, but beware of formats like PF_G8 (that still seem to return alpha in some cases)
+			return true;
+		default:
+			return false;
+	}
+}
+
+void FGLTFTextureUtilities::FullyLoad(const UTexture* InTexture)
+{
+	UTexture* Texture = const_cast<UTexture*>(InTexture);
+
+#if WITH_EDITOR
+	FTextureCompilingManager::Get().FinishCompilation({ Texture });
+#endif
+
+	Texture->SetForceMipLevelsToBeResident(30.0f);
+	Texture->WaitForStreaming();
+}
+
+bool FGLTFTextureUtilities::IsHDR(const UTexture* Texture)
+{
+	switch (Texture->CompressionSettings)
+	{
+		case TC_HDR:
+		case TC_HDR_Compressed:
+		case TC_HalfFloat:
+			return true;
+		default:
+			return false;
+	}
+}
+
+bool FGLTFTextureUtilities::IsCubemap(const UTexture* Texture)
+{
+	return Texture->IsA<UTextureCube>() || Texture->IsA<UTextureRenderTargetCube>();
+}
+
+TextureFilter FGLTFTextureUtilities::GetDefaultFilter(TextureGroup LODGroup)
+{
+	const UTextureLODSettings* TextureLODSettings = UDeviceProfileManager::Get().GetActiveProfile()->GetTextureLODSettings();
+	const ETextureSamplerFilter Filter = TextureLODSettings->GetTextureLODGroup(LODGroup).Filter;
+
+	switch (Filter)
+	{
+		case ETextureSamplerFilter::Point:             return TF_Nearest;
+		case ETextureSamplerFilter::Bilinear:          return TF_Bilinear;
+		case ETextureSamplerFilter::Trilinear:         return TF_Trilinear;
+		case ETextureSamplerFilter::AnisotropicPoint:  return TF_Trilinear; // A lot of engine code doesn't result in nearest
+		case ETextureSamplerFilter::AnisotropicLinear: return TF_Trilinear;
+		default:                                       return TF_MAX;
+	}
+}
+
+int32 FGLTFTextureUtilities::GetMipBias(const UTexture* Texture)
+{
+	if (const UTexture2D* Texture2D = Cast<UTexture2D>(Texture))
+	{
+		return Texture2D->GetNumMips() - Texture2D->GetNumMipsAllowed(true);
+	}
+
+	return Texture->GetCachedLODBias();
+}
+
+FIntPoint FGLTFTextureUtilities::GetInGameSize(const UTexture* Texture)
+{
+	const int32 Width = Texture->GetSurfaceWidth();
+	const int32 Height = Texture->GetSurfaceHeight();
+
+	const int32 MipBias = GetMipBias(Texture);
+
+	const int32 InGameWidth = FMath::Max(Width >> MipBias, 1);
+	const int32 InGameHeight = FMath::Max(Height >> MipBias, 1);
+
+	return { InGameWidth, InGameHeight };
+}
+
+UTextureRenderTarget2D* FGLTFTextureUtilities::CreateRenderTarget(const FIntPoint& Size, bool bHDR)
+{
+	// TODO: instead of PF_FloatRGBA (i.e. RTF_RGBA16f) use PF_A32B32G32R32F (i.e. RTF_RGBA32f) to avoid accuracy loss
+	const EPixelFormat PixelFormat = bHDR ? PF_FloatRGBA : PF_B8G8R8A8;
+
+	// NOTE: both bForceLinearGamma and TargetGamma=2.2 seem necessary for exported images to match their source data.
+	// It's not entirely clear why gamma must be 2.2 (instead of 0.0) and why bInForceLinearGamma must also be true.
+	constexpr bool bForceLinearGamma = true;
+	constexpr float TargetGamma = 2.2f;
+
+	UTextureRenderTarget2D* RenderTarget = NewObject<UTextureRenderTarget2D>();
+	RenderTarget->InitCustomFormat(Size.X, Size.Y, PixelFormat, bForceLinearGamma);
+
+	RenderTarget->TargetGamma = TargetGamma;
+	return RenderTarget;
+}
+
+bool FGLTFTextureUtilities::DrawTexture(UTextureRenderTarget2D* OutTarget, const UTexture2D* InSource)
+{
+	FRenderTarget* RenderTarget = OutTarget->GameThread_GetRenderTargetResource();
+	if (RenderTarget == nullptr)
+	{
+		return false;
+	}
+
+	TRefCountPtr<FBatchedElementParameters> BatchedElementParameters;
+
+	if (InSource->IsNormalMap())
+	{
+		BatchedElementParameters = new FGLTFNormalMapPreview();
+	}
+	else if (IsHDR(InSource))
+	{
+		// NOTE: Simple preview parameters are used to prevent any modifications
+		// such as gamma-correction from being applied during rendering.
+		BatchedElementParameters = new FGLTFSimpleTexture2DPreview();
+	}
+
+	FCanvas Canvas(RenderTarget, nullptr, FGameTime::CreateDilated(0.0f, 0.0f, 0.0f, 0.0f), GMaxRHIFeatureLevel);
+	FCanvasTileItem TileItem(FVector2D::ZeroVector, InSource->GetResource(), FLinearColor::White);
+
+	TileItem.BatchedElementParameters = BatchedElementParameters;
+	TileItem.Draw(&Canvas);
+
+	Canvas.Flush_GameThread();
+	FlushRenderingCommands();
+	Canvas.SetRenderTarget_GameThread(nullptr);
+	FlushRenderingCommands();
+
+	return true;
+}
+
+bool FGLTFTextureUtilities::ReadPixels(const UTextureRenderTarget2D* InRenderTarget, TArray<FColor>& OutPixels)
+{
+	FTextureRenderTarget2DResource* Resource = static_cast<FTextureRenderTarget2DResource*>(const_cast<UTextureRenderTarget2D*>(InRenderTarget)->GetResource());
+	if (Resource == nullptr)
+	{
+		return false;
+	}
+
+	FReadSurfaceDataFlags ReadSurfaceDataFlags(RCM_UNorm, CubeFace_MAX);
+	ReadSurfaceDataFlags.SetLinearToGamma(false);
+	return Resource->ReadPixels(OutPixels, ReadSurfaceDataFlags);
+}
+
+void FGLTFTextureUtilities::FlipGreenChannel(TArray<FColor>& Pixels)
+{
+	for (FColor& Pixel: Pixels)
+	{
+		Pixel.G = 255 - Pixel.G;
+	}
+}
+
+void FGLTFTextureUtilities::TransformColorSpace(TArray<FColor>& Pixels, bool bFromSRGB, bool bToSRGB)
+{
+	if (bFromSRGB == bToSRGB)
+	{
+		return;
+	}
+
+	if (bToSRGB)
+	{
+		for (FColor& Pixel: Pixels)
+		{
+			Pixel = Pixel.ReinterpretAsLinear().ToFColor(true);
+		}
+	}
+	else
+	{
+		for (FColor& Pixel: Pixels)
+		{
+			Pixel = FLinearColor(Pixel).ToFColor(false);
+		}
+	}
+}
diff --git a/Engine/Plugins/Enterprise/GLTFExporter/Source/GLTFExporter/Private/Converters/GLTFTextureUtilities.h b/Engine/Plugins/Enterprise/GLTFExporter/Source/GLTFExporter/Private/Converters/GLTFTextureUtilities.h
new file mode 100644
--- /dev/null
+++ b/Engine/Plugins/Enterprise/GLTFExporter/Source/GLTFExporter/Private/Converters/GLTFTextureUtilities.h
@@ -0,0 +1,35 @@
+// Copyright Epic Games, Inc. All Rights Reserved.
+
+#pragma once
+
+#include "PixelFormat.h"
+#include "Engine/TextureDefines.h"
+
+class UTexture;
+class UTexture2D;
+class UTextureRenderTarget2D;
+
+struct FGLTFTextureUtilities
+{
+	static bool IsAlphaless(EPixelFormat PixelFormat);
+
+	static void FullyLoad(const UTexture* InTexture);
+
+	static bool IsHDR(const UTexture* Texture);
+	static bool IsCubemap(const UTexture* Texture);
+
+	static TextureFilter GetDefaultFilter(TextureGroup Group);
+
+	static int32 GetMipBias(const UTexture* Texture);
+
+	static FIntPoint GetInGameSize(const UTexture* Texture);
+
+	static UTextureRenderTarget2D* CreateRenderTarget(const FIntPoint& Size, bool bIsHDR);
+
+	static bool DrawTexture(UTextureRenderTarget2D* OutTarget, const UTexture2D* InSource);
+
+	static bool ReadPixels(const UTextureRenderTarget2D* InRenderTarget, TArray<FColor>& OutPixels);
+
+	static void FlipGreenChannel(TArray<FColor>& Pixels);
+	static void TransformColorSpace(TArray<FColor>& Pixels, bool bFromSRGB, bool bToSRGB);
+};
diff --git a/Engine/Plugins/Enterprise/GLTFExporter/Source/GLTFExporter/Private/Converters/GLTFTextureUtility.cpp b/Engine/Plugins/Enterprise/GLTFExporter/Source/GLTFExporter/Private/Converters/GLTFTextureUtility.cpp
deleted file mode 100644
diff --git a/Engine/Plugins/Enterprise/GLTFExporter/Source/GLTFExporter/Private/Converters/GLTFTextureUtility.h b/Engine/Plugins/Enterprise/GLTFExporter/Source/GLTFExporter/Private/Converters/GLTFTextureUtility.h
deleted file mode 100644
diff --git a/Engine/Plugins/Enterprise/GLTFExporter/Source/GLTFExporter/Private/Converters/GLTFVariantUtilities.cpp b/Engine/Plugins/Enterprise/GLTFExporter/Source/GLTFExporter/Private/Converters/GLTFVariantUtilities.cpp
new file mode 100644
--- /dev/null
+++ b/Engine/Plugins/Enterprise/GLTFExporter/Source/GLTFExporter/Private/Converters/GLTFVariantUtilities.cpp
@@ -0,0 +1,67 @@
+// Copyright Epic Games, Inc. All Rights Reserved.
+
+#include "Converters/GLTFVariantUtilities.h"
+#include "Builders/GLTFConvertBuilder.h"
+#include "LevelVariantSets.h"
+#include "VariantSet.h"
+#include "Variant.h"
+#include "VariantObjectBinding.h"
+#include "PropertyValue.h"
+
+bool FGLTFVariantUtilities::TryGetPropertyValue(UPropertyValue* Property, void* OutData, uint32 OutSize)
+{
+	if (Property == nullptr || !Property->HasRecordedData())
+	{
+		return false;
+	}
+
+	const TArray<uint8>& RecordedData = Property->GetRecordedData();
+	check(OutSize == RecordedData.Num());
+
+	FMemory::Memcpy(OutData, RecordedData.GetData(), OutSize);
+	return true;
+}
+
+FString FGLTFVariantUtilities::GetLogContext(const UPropertyValue* Property)
+{
+	const UVariantObjectBinding* Parent = Property->GetParent();
+	return GetLogContext(Parent) + TEXT("/") + Property->GetFullDisplayString();
+}
+
+FString FGLTFVariantUtilities::GetLogContext(const UVariantObjectBinding* Binding)
+{
+	const UVariant* Parent = const_cast<UVariantObjectBinding*>(Binding)->GetParent();
+	return GetLogContext(Parent) + TEXT("/") + Binding->GetDisplayText().ToString();
+}
+
+FString FGLTFVariantUtilities::GetLogContext(const UVariant* Variant)
+{
+	const UVariantSet* Parent = const_cast<UVariant*>(Variant)->GetParent();
+	return GetLogContext(Parent) + TEXT("/") + Variant->GetDisplayText().ToString();
+}
+
+FString FGLTFVariantUtilities::GetLogContext(const UVariantSet* VariantSet)
+{
+	const ULevelVariantSets* Parent = const_cast<UVariantSet*>(VariantSet)->GetParent();
+	return GetLogContext(Parent) + TEXT("/") + VariantSet->GetDisplayText().ToString();
+}
+
+FString FGLTFVariantUtilities::GetLogContext(const ULevelVariantSets* LevelVariantSets)
+{
+	return LevelVariantSets->GetName();
+}
+
+FGLTFJsonMaterial* FGLTFVariantUtilities::AddUniqueMaterial(FGLTFConvertBuilder& Builder, const UMaterialInterface* Material, const UMeshComponent* MeshComponent, int32 MaterialIndex)
+{
+	if (Builder.ExportOptions->ExportMaterialVariants == EGLTFMaterialVariantMode::UseMeshData)
+	{
+		if (Builder.ExportOptions->BakeMaterialInputs == EGLTFMaterialBakeMode::UseMeshData)
+		{
+			return Builder.AddUniqueMaterial(Material, MeshComponent, -1, MaterialIndex);
+		}
+
+		// TODO: report warning (about materials won't be export using mesh data because BakeMaterialInputs is not set to UseMeshData)
+	}
+
+	return Builder.AddUniqueMaterial(Material);
+}
diff --git a/Engine/Plugins/Enterprise/GLTFExporter/Source/GLTFExporter/Private/Converters/GLTFVariantUtilities.h b/Engine/Plugins/Enterprise/GLTFExporter/Source/GLTFExporter/Private/Converters/GLTFVariantUtilities.h
new file mode 100644
--- /dev/null
+++ b/Engine/Plugins/Enterprise/GLTFExporter/Source/GLTFExporter/Private/Converters/GLTFVariantUtilities.h
@@ -0,0 +1,34 @@
+// Copyright Epic Games, Inc. All Rights Reserved.
+
+#pragma once
+
+#include "Json/GLTFJsonCore.h"
+
+class FGLTFConvertBuilder;
+struct FCapturedPropSegment;
+class ULevelVariantSets;
+class UMaterialInterface;
+class UMeshComponent;
+class UVariantSet;
+class UVariant;
+class UVariantObjectBinding;
+class UPropertyValue;
+
+struct FGLTFVariantUtilities
+{
+	template <typename T>
+	static bool TryGetPropertyValue(UPropertyValue* Property, T& OutValue)
+	{
+		return TryGetPropertyValue(Property, &OutValue, sizeof(T));
+	}
+
+	static bool TryGetPropertyValue(UPropertyValue* Property, void* OutData, uint32 OutSize);
+
+	static FString GetLogContext(const UPropertyValue* Property);
+	static FString GetLogContext(const UVariantObjectBinding* Binding);
+	static FString GetLogContext(const UVariant* Variant);
+	static FString GetLogContext(const UVariantSet* VariantSet);
+	static FString GetLogContext(const ULevelVariantSets* LevelVariantSets);
+
+	static FGLTFJsonMaterial* AddUniqueMaterial(FGLTFConvertBuilder& Builder, const UMaterialInterface* Material, const UMeshComponent* MeshComponent, int32 MaterialIndex);
+};
diff --git a/Engine/Plugins/Enterprise/GLTFExporter/Source/GLTFExporter/Private/Converters/GLTFVariantUtility.cpp b/Engine/Plugins/Enterprise/GLTFExporter/Source/GLTFExporter/Private/Converters/GLTFVariantUtility.cpp
deleted file mode 100644
diff --git a/Engine/Plugins/Enterprise/GLTFExporter/Source/GLTFExporter/Private/Converters/GLTFVariantUtility.h b/Engine/Plugins/Enterprise/GLTFExporter/Source/GLTFExporter/Private/Converters/GLTFVariantUtility.h
deleted file mode 100644
diff --git a/Engine/Plugins/Enterprise/GLTFExporter/Source/GLTFExporter/Private/Exporters/GLTFAnimSequenceExporter.cpp b/Engine/Plugins/Enterprise/GLTFExporter/Source/GLTFExporter/Private/Exporters/GLTFAnimSequenceExporter.cpp
--- a/Engine/Plugins/Enterprise/GLTFExporter/Source/GLTFExporter/Private/Exporters/GLTFAnimSequenceExporter.cpp
+++ b/Engine/Plugins/Enterprise/GLTFExporter/Source/GLTFExporter/Private/Exporters/GLTFAnimSequenceExporter.cpp
@@ -1,7 +1,7 @@
 // Copyright Epic Games, Inc. All Rights Reserved.
 
 #include "Exporters/GLTFAnimSequenceExporter.h"
-#include "Exporters/GLTFExporterUtility.h"
+#include "Exporters/GLTFExporterUtilities.h"
 #include "Builders/GLTFContainerBuilder.h"
 #include "Animation/AnimSequence.h"
 #include "Engine/SkeletalMesh.h"
@@ -34,7 +34,7 @@ bool UGLTFAnimSequenceExporter::AddObject(FGLTFContainerBuilder& Builder, const
 		return false;
 	}
 
-	const USkeletalMesh* SkeletalMesh = FGLTFExporterUtility::GetPreviewMesh(AnimSequence);
+	const USkeletalMesh* SkeletalMesh = FGLTFExporterUtilities::GetPreviewMesh(AnimSequence);
 	if (SkeletalMesh == nullptr)
 	{
 		Builder.LogError(
diff --git a/Engine/Plugins/Enterprise/GLTFExporter/Source/GLTFExporter/Private/Exporters/GLTFExporter.cpp b/Engine/Plugins/Enterprise/GLTFExporter/Source/GLTFExporter/Private/Exporters/GLTFExporter.cpp
--- a/Engine/Plugins/Enterprise/GLTFExporter/Source/GLTFExporter/Private/Exporters/GLTFExporter.cpp
+++ b/Engine/Plugins/Enterprise/GLTFExporter/Source/GLTFExporter/Private/Exporters/GLTFExporter.cpp
@@ -2,7 +2,7 @@
 
 #include "Exporters/GLTFExporter.h"
 #include "Engine/Engine.h"
-#include "Exporters/GLTFExporterUtility.h"
+#include "Exporters/GLTFExporterUtilities.h"
 #include "Exporters/GLTFExporterAnalytics.h"
 #include "Misc/Paths.h"
 #include "Options/GLTFExportOptions.h"
@@ -47,7 +47,7 @@ bool UGLTFExporter::ExportBinary(UObject* Object, const TCHAR* Type, FArchive& A
 	TSet<AActor*> SelectedActors;
 	if (bSelectedOnly)
 	{
-		FGLTFExporterUtility::GetSelectedActors(SelectedActors);
+		FGLTFExporterUtilities::GetSelectedActors(SelectedActors);
 	}
 
 	// TODO: add support for UAssetExportTask::IgnoreObjectList?
diff --git a/Engine/Plugins/Enterprise/GLTFExporter/Source/GLTFExporter/Private/Exporters/GLTFExporterUtilities.cpp b/Engine/Plugins/Enterprise/GLTFExporter/Source/GLTFExporter/Private/Exporters/GLTFExporterUtilities.cpp
new file mode 100644
--- /dev/null
+++ b/Engine/Plugins/Enterprise/GLTFExporter/Source/GLTFExporter/Private/Exporters/GLTFExporterUtilities.cpp
@@ -0,0 +1,127 @@
+// Copyright Epic Games, Inc. All Rights Reserved.
+
+#include "Exporters/GLTFExporterUtilities.h"
+#include "AssetRegistry/AssetData.h"
+#include "Materials/MaterialInstance.h"
+#include "AssetRegistry/AssetRegistryModule.h"
+#include "Animation/AnimSequence.h"
+#include "Animation/Skeleton.h"
+#include "Engine/SkeletalMesh.h"
+#include "Engine/StaticMesh.h"
+#include "Engine/Selection.h"
+#include "Engine/World.h"
+#if WITH_EDITOR
+#include "Editor.h"
+#endif
+
+void FGLTFExporterUtilities::GetSelectedActors(TSet<AActor*>& OutSelectedActors)
+{
+#if WITH_EDITOR
+	USelection* Selection = GEditor->GetSelectedActors();
+	for (FSelectionIterator SelectedObject(*Selection); SelectedObject; ++SelectedObject)
+	{
+		if (AActor* SelectedActor = Cast<AActor>(*SelectedObject))
+		{
+			OutSelectedActors.Add(SelectedActor);
+		}
+	}
+#endif
+}
+
+const UStaticMesh* FGLTFExporterUtilities::GetPreviewMesh(const UMaterialInterface* Material)
+{
+#if WITH_EDITORONLY_DATA
+	do
+	{
+		const UStaticMesh* PreviewMesh = Cast<UStaticMesh>(Material->PreviewMesh.TryLoad());
+		if (PreviewMesh != nullptr)
+		{
+			return PreviewMesh;
+		}
+
+		const UMaterialInstance* MaterialInstance = Cast<UMaterialInstance>(Material);
+		Material = MaterialInstance != nullptr ? MaterialInstance->Parent : nullptr;
+	} while (Material != nullptr);
+#endif
+
+	static const UStaticMesh* DefaultPreviewMesh = GetSphereMesh();
+	return DefaultPreviewMesh;
+}
+
+const UStaticMesh* FGLTFExporterUtilities::GetSphereMesh()
+{
+	const UStaticMesh* SphereMesh = LoadObject<UStaticMesh>(nullptr, TEXT("/Engine/EditorMeshes/EditorSphere.EditorSphere"));
+	if (SphereMesh == nullptr)
+	{
+		SphereMesh = LoadObject<UStaticMesh>(nullptr, TEXT("/Engine/BasicShapes/Sphere.Sphere"));
+		if (SphereMesh == nullptr)
+		{
+			SphereMesh = LoadObject<UStaticMesh>(nullptr, TEXT("/Engine/EngineMeshes/Sphere.Sphere"));
+		}
+	}
+
+	return SphereMesh;
+}
+
+const USkeletalMesh* FGLTFExporterUtilities::GetPreviewMesh(const UAnimSequence* AnimSequence)
+{
+	const USkeletalMesh* PreviewMesh = AnimSequence->GetPreviewMesh();
+	if (PreviewMesh == nullptr)
+	{
+		const USkeleton* Skeleton = AnimSequence->GetSkeleton();
+		if (Skeleton != nullptr)
+		{
+			PreviewMesh = Skeleton->GetPreviewMesh();
+			if (PreviewMesh == nullptr)
+			{
+				PreviewMesh = FindCompatibleMesh(Skeleton);
+			}
+		}
+	}
+
+	return PreviewMesh;
+}
+
+const USkeletalMesh* FGLTFExporterUtilities::FindCompatibleMesh(const USkeleton *Skeleton)
+{
+	const FName SkeletonMemberName = USkeletalMesh::GetSkeletonMemberName();
+
+	FARFilter Filter;
+	Filter.ClassPaths.Add(USkeletalMesh::StaticClass()->GetClassPathName());
+	Filter.TagsAndValues.Add(SkeletonMemberName, FAssetData(Skeleton).GetExportTextName());
+
+	TArray<FAssetData> FilteredAssets;
+	FAssetRegistryModule::GetRegistry().GetAssets(Filter, FilteredAssets);
+
+	for (const FAssetData& Asset : FilteredAssets)
+	{
+		const USkeletalMesh* SkeletalMesh = Cast<USkeletalMesh>(Asset.GetAsset());
+		if (SkeletalMesh != nullptr)
+		{
+			return SkeletalMesh;
+		}
+	}
+
+	return nullptr;
+}
+
+TArray<UWorld*> FGLTFExporterUtilities::GetAssociatedWorlds(const UObject* Object)
+{
+	TArray<UWorld*> Worlds;
+	TArray<FAssetIdentifier> Dependencies;
+
+	const FName OuterPathName = *Object->GetOutermost()->GetPathName();
+	FAssetRegistryModule::GetRegistry().GetDependencies(OuterPathName, Dependencies);
+
+	for (FAssetIdentifier& Dependency : Dependencies)
+	{
+		FString PackageName = Dependency.PackageName.ToString();
+		UWorld* World = LoadObject<UWorld>(nullptr, *PackageName, nullptr, LOAD_NoWarn);
+		if (World != nullptr)
+		{
+			Worlds.AddUnique(World);
+		}
+	}
+
+	return Worlds;
+}
diff --git a/Engine/Plugins/Enterprise/GLTFExporter/Source/GLTFExporter/Private/Exporters/GLTFExporterUtilities.h b/Engine/Plugins/Enterprise/GLTFExporter/Source/GLTFExporter/Private/Exporters/GLTFExporterUtilities.h
new file mode 100644
--- /dev/null
+++ b/Engine/Plugins/Enterprise/GLTFExporter/Source/GLTFExporter/Private/Exporters/GLTFExporterUtilities.h
@@ -0,0 +1,25 @@
+// Copyright Epic Games, Inc. All Rights Reserved.
+
+#pragma once
+
+#include "CoreMinimal.h"
+
+class AActor;
+class USkeletalMesh;
+class UStaticMesh;
+class UMaterialInterface;
+class UAnimSequence;
+class USkeleton;
+
+struct FGLTFExporterUtilities
+{
+	static void GetSelectedActors(TSet<AActor*>& OutSelectedActors);
+
+	static const UStaticMesh* GetPreviewMesh(const UMaterialInterface* Material);
+	static const UStaticMesh* GetSphereMesh();
+
+	static const USkeletalMesh* GetPreviewMesh(const UAnimSequence* AnimSequence);
+	static const USkeletalMesh* FindCompatibleMesh(const USkeleton* Skeleton);
+
+	static TArray<UWorld*> GetAssociatedWorlds(const UObject* Object);
+};
diff --git a/Engine/Plugins/Enterprise/GLTFExporter/Source/GLTFExporter/Private/Exporters/GLTFExporterUtility.cpp b/Engine/Plugins/Enterprise/GLTFExporter/Source/GLTFExporter/Private/Exporters/GLTFExporterUtility.cpp
deleted file mode 100644
diff --git a/Engine/Plugins/Enterprise/GLTFExporter/Source/GLTFExporter/Private/Exporters/GLTFExporterUtility.h b/Engine/Plugins/Enterprise/GLTFExporter/Source/GLTFExporter/Private/Exporters/GLTFExporterUtility.h
deleted file mode 100644
diff --git a/Engine/Plugins/Enterprise/GLTFExporter/Source/GLTFExporter/Private/Exporters/GLTFLevelSequenceExporter.cpp b/Engine/Plugins/Enterprise/GLTFExporter/Source/GLTFExporter/Private/Exporters/GLTFLevelSequenceExporter.cpp
--- a/Engine/Plugins/Enterprise/GLTFExporter/Source/GLTFExporter/Private/Exporters/GLTFLevelSequenceExporter.cpp
+++ b/Engine/Plugins/Enterprise/GLTFExporter/Source/GLTFExporter/Private/Exporters/GLTFLevelSequenceExporter.cpp
@@ -1,7 +1,7 @@
 // Copyright Epic Games, Inc. All Rights Reserved.
 
 #include "Exporters/GLTFLevelSequenceExporter.h"
-#include "Exporters/GLTFExporterUtility.h"
+#include "Exporters/GLTFExporterUtilities.h"
 #include "Builders/GLTFContainerBuilder.h"
 #include "LevelSequence.h"
 
@@ -25,7 +25,7 @@ bool UGLTFLevelSequenceExporter::AddObject(FGLTFContainerBuilder& Builder, const
 		return false;
 	}
 
-	TArray<UWorld*> Worlds = FGLTFExporterUtility::GetAssociatedWorlds(LevelSequence);
+	TArray<UWorld*> Worlds = FGLTFExporterUtilities::GetAssociatedWorlds(LevelSequence);
 	if (Worlds.Num() == 0)
 	{
 		Builder.LogError(
diff --git a/Engine/Plugins/Enterprise/GLTFExporter/Source/GLTFExporter/Private/Exporters/GLTFLevelVariantSetsExporter.cpp b/Engine/Plugins/Enterprise/GLTFExporter/Source/GLTFExporter/Private/Exporters/GLTFLevelVariantSetsExporter.cpp
--- a/Engine/Plugins/Enterprise/GLTFExporter/Source/GLTFExporter/Private/Exporters/GLTFLevelVariantSetsExporter.cpp
+++ b/Engine/Plugins/Enterprise/GLTFExporter/Source/GLTFExporter/Private/Exporters/GLTFLevelVariantSetsExporter.cpp
@@ -1,7 +1,7 @@
 // Copyright Epic Games, Inc. All Rights Reserved.
 
 #include "Exporters/GLTFLevelVariantSetsExporter.h"
-#include "Exporters/GLTFExporterUtility.h"
+#include "Exporters/GLTFExporterUtilities.h"
 #include "Builders/GLTFContainerBuilder.h"
 #include "Engine/World.h"
 #include "LevelVariantSets.h"
@@ -26,7 +26,7 @@ bool UGLTFLevelVariantSetsExporter::AddObject(FGLTFContainerBuilder& Builder, co
 		return false;
 	}
 
-	TArray<UWorld*> Worlds = FGLTFExporterUtility::GetAssociatedWorlds(LevelVariantSets);
+	TArray<UWorld*> Worlds = FGLTFExporterUtilities::GetAssociatedWorlds(LevelVariantSets);
 	if (Worlds.Num() == 0)
 	{
 		Builder.LogError(
diff --git a/Engine/Plugins/Enterprise/GLTFExporter/Source/GLTFExporter/Private/Exporters/GLTFMaterialExporter.cpp b/Engine/Plugins/Enterprise/GLTFExporter/Source/GLTFExporter/Private/Exporters/GLTFMaterialExporter.cpp
--- a/Engine/Plugins/Enterprise/GLTFExporter/Source/GLTFExporter/Private/Exporters/GLTFMaterialExporter.cpp
+++ b/Engine/Plugins/Enterprise/GLTFExporter/Source/GLTFExporter/Private/Exporters/GLTFMaterialExporter.cpp
@@ -1,7 +1,7 @@
 // Copyright Epic Games, Inc. All Rights Reserved.
 
 #include "Exporters/GLTFMaterialExporter.h"
-#include "Exporters/GLTFExporterUtility.h"
+#include "Exporters/GLTFExporterUtilities.h"
 #include "Builders/GLTFContainerBuilder.h"
 #include "Engine/StaticMesh.h"
 #include "Materials/MaterialInterface.h"
@@ -23,7 +23,7 @@ bool UGLTFMaterialExporter::AddObject(FGLTFContainerBuilder& Builder, const UObj
 
 	if (Builder.ExportOptions->bExportPreviewMesh)
 	{
-		const UStaticMesh* PreviewMesh = FGLTFExporterUtility::GetPreviewMesh(Material);
+		const UStaticMesh* PreviewMesh = FGLTFExporterUtilities::GetPreviewMesh(Material);
 		if (PreviewMesh != nullptr)
 		{
 			FGLTFJsonMesh* Mesh = Builder.AddUniqueMesh(PreviewMesh, { Material });
diff --git a/Engine/Plugins/Enterprise/GLTFExporter/Source/GLTFExporter/Private/Tasks/GLTFDelayedAnimationTasks.cpp b/Engine/Plugins/Enterprise/GLTFExporter/Source/GLTFExporter/Private/Tasks/GLTFDelayedAnimationTasks.cpp
--- a/Engine/Plugins/Enterprise/GLTFExporter/Source/GLTFExporter/Private/Tasks/GLTFDelayedAnimationTasks.cpp
+++ b/Engine/Plugins/Enterprise/GLTFExporter/Source/GLTFExporter/Private/Tasks/GLTFDelayedAnimationTasks.cpp
@@ -4,7 +4,7 @@
 #include "Builders/GLTFContainerBuilder.h"
 #include "Engine/Level.h"
 #include "Utilities/GLTFCoreUtilities.h"
-#include "Converters/GLTFBoneUtility.h"
+#include "Converters/GLTFBoneUtilities.h"
 #include "LevelSequence.h"
 #include "LevelSequenceActor.h"
 #include "LevelSequencePlayer.h"
@@ -22,7 +22,7 @@ void FGLTFDelayedAnimSequenceTask::Process()
 	JsonAnimation->Name = AnimSequence->GetName() + TEXT("_") + FString::FromInt(JsonAnimation->Index); // Ensure unique name due to limitation in certain gltf viewers
 
 	TArray<float> Timestamps;
-	FGLTFBoneUtility::GetFrameTimestamps(AnimSequence, Timestamps);
+	FGLTFBoneUtilities::GetFrameTimestamps(AnimSequence, Timestamps);
 	const int32 FrameCount = Timestamps.Num();
 
 	// TODO: add animation data accessor converters to reuse track information
@@ -44,10 +44,10 @@ void FGLTFDelayedAnimSequenceTask::Process()
 	}
 
 	TArray<FBoneIndexType> BoneIndices;
-	FGLTFBoneUtility::GetBoneIndices(AnimSequence->GetSkeleton(), BoneIndices);
+	FGLTFBoneUtilities::GetBoneIndices(AnimSequence->GetSkeleton(), BoneIndices);
 
 	TArray<TArray<FTransform>> FrameTransforms;
-	FGLTFBoneUtility::GetBoneTransformsByFrame(AnimSequence, Timestamps, BoneIndices, FrameTransforms);
+	FGLTFBoneUtilities::GetBoneTransformsByFrame(AnimSequence, Timestamps, BoneIndices, FrameTransforms);
 
 	for (const FBoneIndexType BoneIndex : BoneIndices)
 	{
diff --git a/Engine/Plugins/Enterprise/GLTFExporter/Source/GLTFExporter/Private/Tasks/GLTFDelayedMaterialTasks.cpp b/Engine/Plugins/Enterprise/GLTFExporter/Source/GLTFExporter/Private/Tasks/GLTFDelayedMaterialTasks.cpp
--- a/Engine/Plugins/Enterprise/GLTFExporter/Source/GLTFExporter/Private/Tasks/GLTFDelayedMaterialTasks.cpp
+++ b/Engine/Plugins/Enterprise/GLTFExporter/Source/GLTFExporter/Private/Tasks/GLTFDelayedMaterialTasks.cpp
@@ -2,8 +2,8 @@
 
 #include "Tasks/GLTFDelayedMaterialTasks.h"
 #include "Utilities/GLTFCoreUtilities.h"
-#include "Converters/GLTFNameUtility.h"
-#include "Converters/GLTFMaterialUtility.h"
+#include "Converters/GLTFNameUtilities.h"
+#include "Converters/GLTFMaterialUtilities.h"
 #include "Builders/GLTFContainerBuilder.h"
 #include "Utilities/GLTFProxyMaterialUtilities.h"
 #include "MaterialDomain.h"
@@ -228,7 +228,7 @@ void FGLTFDelayedMaterialTask::GetProxyParameters(FGLTFJsonMaterial& OutMaterial
 
 		if (HasProxyParameter(FGLTFProxyMaterialInfo::Normal.Texture))
 		{
-			if (OutMaterial.ShadingModel == EGLTFJsonShadingModel::ClearCoat && !FGLTFMaterialUtility::IsClearCoatBottomNormalEnabled())
+			if (OutMaterial.ShadingModel == EGLTFJsonShadingModel::ClearCoat && !FGLTFMaterialUtilities::IsClearCoatBottomNormalEnabled())
 			{
 				Builder.LogWarning(FString::Printf(
 					TEXT("Proxy material %s won't be exported with (bottom) normal because ClearCoatEnableSecondNormal in project rendering settings is disabled"),
@@ -340,7 +340,7 @@ EMaterialShadingModel FGLTFDelayedMaterialTask::GetShadingModel() const
 		Builder.LogWarning(FString::Printf(
 			TEXT("No shading model defined for material %s, will export as %s"),
 			*Material->GetName(),
-			*FGLTFNameUtility::GetName(ShadingModel)));
+			*FGLTFNameUtilities::GetName(ShadingModel)));
 		return ShadingModel;
 	}
 
@@ -349,37 +349,37 @@ EMaterialShadingModel FGLTFDelayedMaterialTask::GetShadingModel() const
 #if WITH_EDITOR
 		if (!FApp::CanEverRender())
 		{
-			const EMaterialShadingModel ShadingModel = FGLTFMaterialUtility::GetRichestShadingModel(Possibilities);
+			const EMaterialShadingModel ShadingModel = FGLTFMaterialUtilities::GetRichestShadingModel(Possibilities);
 			Builder.LogWarning(FString::Printf(
 				TEXT("Can't evaluate shading model expression in material %s because renderer missing, will export as %s"),
 				*Material->GetName(),
-				*FGLTFNameUtility::GetName(ShadingModel)));
+				*FGLTFNameUtilities::GetName(ShadingModel)));
 			return ShadingModel;
 		}
 
 		if (Material->IsShadingModelFromMaterialExpression())
 		{
-			const FMaterialShadingModelField Evaluation = FGLTFMaterialUtility::EvaluateShadingModelExpression(Material);
+			const FMaterialShadingModelField Evaluation = FGLTFMaterialUtilities::EvaluateShadingModelExpression(Material);
 			const int32 EvaluationCount = Evaluation.CountShadingModels();
 
 			if (EvaluationCount == 0)
 			{
-				const EMaterialShadingModel ShadingModel = FGLTFMaterialUtility::GetRichestShadingModel(Possibilities);
+				const EMaterialShadingModel ShadingModel = FGLTFMaterialUtilities::GetRichestShadingModel(Possibilities);
 				Builder.LogWarning(FString::Printf(
 					TEXT("Evaluation of shading model expression in material %s returned none, will export as %s"),
 					*Material->GetName(),
-					*FGLTFNameUtility::GetName(ShadingModel)));
+					*FGLTFNameUtilities::GetName(ShadingModel)));
 				return ShadingModel;
 			}
 
 			if (EvaluationCount > 1)
 			{
-				const EMaterialShadingModel ShadingModel = FGLTFMaterialUtility::GetRichestShadingModel(Evaluation);
+				const EMaterialShadingModel ShadingModel = FGLTFMaterialUtilities::GetRichestShadingModel(Evaluation);
 				Builder.LogWarning(FString::Printf(
 					TEXT("Evaluation of shading model expression in material %s is inconclusive (%s), will export as %s"),
 					*Material->GetName(),
-					*FGLTFMaterialUtility::ShadingModelsToString(Evaluation),
-					*FGLTFNameUtility::GetName(ShadingModel)));
+					*FGLTFMaterialUtilities::ShadingModelsToString(Evaluation),
+					*FGLTFNameUtilities::GetName(ShadingModel)));
 				return ShadingModel;
 			}
 
@@ -388,11 +388,11 @@ EMaterialShadingModel FGLTFDelayedMaterialTask::GetShadingModel() const
 
 		// we should never end up here
 #else
-		const EMaterialShadingModel ShadingModel = FGLTFMaterialUtility::GetRichestShadingModel(Possibilities);
+		const EMaterialShadingModel ShadingModel = FGLTFMaterialUtilities::GetRichestShadingModel(Possibilities);
 		Builder.LogWarning(FString::Printf(
 			TEXT("Can't evaluate shading model expression in material %s without editor, will export as %s"),
 			*Material->GetName(),
-			*FGLTFNameUtility::GetName(ShadingModel)));
+			*FGLTFNameUtilities::GetName(ShadingModel)));
 		return ShadingModel;
 #endif
 	}
@@ -418,9 +418,9 @@ void FGLTFDelayedMaterialTask::ConvertShadingModel(EGLTFJsonShadingModel& OutSha
 
 		Builder.LogWarning(FString::Printf(
 			TEXT("Unsupported shading model (%s) in material %s, will export as %s"),
-			*FGLTFNameUtility::GetName(ShadingModel),
+			*FGLTFNameUtilities::GetName(ShadingModel),
 			*Material->GetName(),
-			*FGLTFNameUtility::GetName(MSM_DefaultLit)));
+			*FGLTFNameUtilities::GetName(MSM_DefaultLit)));
 		return;
 	}
 
@@ -430,9 +430,9 @@ void FGLTFDelayedMaterialTask::ConvertShadingModel(EGLTFJsonShadingModel& OutSha
 
 		Builder.LogWarning(FString::Printf(
 			TEXT("Shading model (%s) in material %s disabled by export options, will export as %s"),
-			*FGLTFNameUtility::GetName(ShadingModel),
+			*FGLTFNameUtilities::GetName(ShadingModel),
 			*Material->GetName(),
-			*FGLTFNameUtility::GetName(MSM_DefaultLit)));
+			*FGLTFNameUtilities::GetName(MSM_DefaultLit)));
 		return;
 	}
 
@@ -442,9 +442,9 @@ void FGLTFDelayedMaterialTask::ConvertShadingModel(EGLTFJsonShadingModel& OutSha
 
 		Builder.LogWarning(FString::Printf(
 			TEXT("Shading model (%s) in material %s disabled by export options, will export as %s"),
-			*FGLTFNameUtility::GetName(ShadingModel),
+			*FGLTFNameUtilities::GetName(ShadingModel),
 			*Material->GetName(),
-			*FGLTFNameUtility::GetName(MSM_DefaultLit)));
+			*FGLTFNameUtilities::GetName(MSM_DefaultLit)));
 	}
 }
 
@@ -459,9 +459,9 @@ void FGLTFDelayedMaterialTask::ConvertAlphaMode(EGLTFJsonAlphaMode& OutAlphaMode
 
 		Builder.LogWarning(FString::Printf(
 			TEXT("Unsupported blend mode (%s) in material %s, will export as %s"),
-			*FGLTFNameUtility::GetName(BlendMode),
+			*FGLTFNameUtilities::GetName(BlendMode),
 			*Material->GetName(),
-			*FGLTFNameUtility::GetName(BLEND_Translucent)));
+			*FGLTFNameUtilities::GetName(BLEND_Translucent)));
 		return;
 	}
 }
@@ -566,7 +566,7 @@ bool FGLTFDelayedMaterialTask::TryGetBaseColorAndOpacity(FGLTFJsonPBRMetallicRou
 			return FColor(BaseColor.R, BaseColor.G, BaseColor.B, Opacity.R);
 		});
 
-	FGLTFJsonTexture* CombinedTexture = FGLTFMaterialUtility::AddTexture(
+	FGLTFJsonTexture* CombinedTexture = FGLTFMaterialUtilities::AddTexture(
 		Builder,
 		CombinedPixels,
 		TextureSize,
@@ -656,7 +656,7 @@ bool FGLTFDelayedMaterialTask::TryGetMetallicAndRoughness(FGLTFJsonPBRMetallicRo
 	}
 
 	int32 CombinedTexCoord;
-	
+
 	// If one is constant, it means we can use the other's texture coords
 	if (MetallicBakeOutput.bIsConstant || MetallicTexCoord == RoughnessTexCoord)
 	{
@@ -682,7 +682,7 @@ bool FGLTFDelayedMaterialTask::TryGetMetallicAndRoughness(FGLTFJsonPBRMetallicRo
 			return FColor(0, Roughness.R, Metallic.R);
 		});
 
-	FGLTFJsonTexture* CombinedTexture = FGLTFMaterialUtility::AddTexture(
+	FGLTFJsonTexture* CombinedTexture = FGLTFMaterialUtilities::AddTexture(
 		Builder,
 		CombinedPixels,
 		TextureSize,
@@ -771,7 +771,7 @@ bool FGLTFDelayedMaterialTask::TryGetClearCoatRoughness(FGLTFJsonClearCoatExtens
 	}
 
 	int32 CombinedTexCoord;
-	
+
 	// If one is constant, it means we can use the other's texture coords
 	if (IntensityBakeOutput.bIsConstant || IntensityTexCoord == RoughnessTexCoord)
 	{
@@ -797,7 +797,7 @@ bool FGLTFDelayedMaterialTask::TryGetClearCoatRoughness(FGLTFJsonClearCoatExtens
 			return FColor(Intensity.R, Roughness.R, 0);
 		});
 
-	FGLTFJsonTexture* CombinedTexture = FGLTFMaterialUtility::AddTexture(
+	FGLTFJsonTexture* CombinedTexture = FGLTFMaterialUtilities::AddTexture(
 		Builder,
 		CombinedPixels,
 		TextureSize,
@@ -880,7 +880,7 @@ bool FGLTFDelayedMaterialTask::IsPropertyNonDefault(const FGLTFMaterialPropertyE
 		}
 	}
 
-	const FExpressionInput* MaterialInput = FGLTFMaterialUtility::GetInputForProperty(Material, Property);
+	const FExpressionInput* MaterialInput = FGLTFMaterialUtilities::GetInputForProperty(Material, Property);
 	if (MaterialInput == nullptr)
 	{
 		return false;
@@ -892,7 +892,7 @@ bool FGLTFDelayedMaterialTask::IsPropertyNonDefault(const FGLTFMaterialPropertyE
 		return false;
 	}
 
-	if (Property == FGLTFMaterialPropertyEx::ClearCoatBottomNormal && !FGLTFMaterialUtility::IsClearCoatBottomNormalEnabled())
+	if (Property == FGLTFMaterialPropertyEx::ClearCoatBottomNormal && !FGLTFMaterialUtilities::IsClearCoatBottomNormalEnabled())
 	{
 		Builder.LogWarning(FString::Printf(
 			TEXT("Material %s won't be exported with clear coat bottom normal because ClearCoatEnableSecondNormal in project rendering settings is disabled"),
@@ -936,7 +936,7 @@ bool FGLTFDelayedMaterialTask::TryGetConstantColor(FLinearColor& OutValue, const
 		return false;
 	}
 
-	const FMaterialInput<FColor>* MaterialInput = FGLTFMaterialUtility::GetInputForProperty<FColor>(Material, Property);
+	const FMaterialInput<FColor>* MaterialInput = FGLTFMaterialUtilities::GetInputForProperty<FColor>(Material, Property);
 	if (MaterialInput == nullptr)
 	{
 		// TODO: report error
@@ -952,7 +952,7 @@ bool FGLTFDelayedMaterialTask::TryGetConstantColor(FLinearColor& OutValue, const
 	const UMaterialExpression* Expression = MaterialInput->Expression;
 	if (Expression == nullptr)
 	{
-		OutValue = FLinearColor(FGLTFMaterialUtility::GetPropertyDefaultValue(Property));
+		OutValue = FLinearColor(FGLTFMaterialUtilities::GetPropertyDefaultValue(Property));
 		return true;
 	}
 
@@ -970,11 +970,11 @@ bool FGLTFDelayedMaterialTask::TryGetConstantColor(FLinearColor& OutValue, const
 			}
 		}
 
-		const uint32 MaskComponentCount = FGLTFMaterialUtility::GetMaskComponentCount(*MaterialInput);
+		const uint32 MaskComponentCount = FGLTFMaterialUtilities::GetMaskComponentCount(*MaterialInput);
 
 		if (MaskComponentCount > 0)
 		{
-			const FLinearColor Mask = FGLTFMaterialUtility::GetMask(*MaterialInput);
+			const FLinearColor Mask = FGLTFMaterialUtilities::GetMask(*MaterialInput);
 
 			Value *= Mask;
 
@@ -1043,7 +1043,7 @@ bool FGLTFDelayedMaterialTask::TryGetConstantScalar(float& OutValue, const FGLTF
 		return false;
 	}
 
-	const FMaterialInput<float>* MaterialInput = FGLTFMaterialUtility::GetInputForProperty<float>(Material, Property);
+	const FMaterialInput<float>* MaterialInput = FGLTFMaterialUtilities::GetInputForProperty<float>(Material, Property);
 	if (MaterialInput == nullptr)
 	{
 		// TODO: report error
@@ -1059,7 +1059,7 @@ bool FGLTFDelayedMaterialTask::TryGetConstantScalar(float& OutValue, const FGLTF
 	const UMaterialExpression* Expression = MaterialInput->Expression;
 	if (Expression == nullptr)
 	{
-		OutValue = FGLTFMaterialUtility::GetPropertyDefaultValue(Property).X;
+		OutValue = FGLTFMaterialUtilities::GetPropertyDefaultValue(Property).X;
 		return true;
 	}
 
@@ -1077,11 +1077,11 @@ bool FGLTFDelayedMaterialTask::TryGetConstantScalar(float& OutValue, const FGLTF
 			}
 		}
 
-		const uint32 MaskComponentCount = FGLTFMaterialUtility::GetMaskComponentCount(*MaterialInput);
+		const uint32 MaskComponentCount = FGLTFMaterialUtilities::GetMaskComponentCount(*MaterialInput);
 
 		if (MaskComponentCount > 0)
 		{
-			const FLinearColor Mask = FGLTFMaterialUtility::GetMask(*MaterialInput);
+			const FLinearColor Mask = FGLTFMaterialUtilities::GetMask(*MaterialInput);
 			Value *= Mask;
 		}
 
@@ -1143,7 +1143,7 @@ bool FGLTFDelayedMaterialTask::TryGetSourceTexture(FGLTFJsonTextureInfo& OutTexI
 
 	if (TryGetSourceTexture(Texture, TexCoord, Transform, Property, AllowedMasks))
 	{
-		OutTexInfo.Index = Builder.AddUniqueTexture(Texture, FGLTFMaterialUtility::IsSRGB(Property));
+		OutTexInfo.Index = Builder.AddUniqueTexture(Texture, FGLTFMaterialUtilities::IsSRGB(Property));
 		OutTexInfo.TexCoord = TexCoord;
 		OutTexInfo.Transform = Transform;
 		return true;
@@ -1159,7 +1159,7 @@ bool FGLTFDelayedMaterialTask::TryGetSourceTexture(const UTexture2D*& OutTexture
 		return false;
 	}
 
-	const FExpressionInput* MaterialInput = FGLTFMaterialUtility::GetInputForProperty(Material, Property);
+	const FExpressionInput* MaterialInput = FGLTFMaterialUtilities::GetInputForProperty(Material, Property);
 	if (MaterialInput == nullptr)
 	{
 		// TODO: report error
@@ -1172,7 +1172,7 @@ bool FGLTFDelayedMaterialTask::TryGetSourceTexture(const UTexture2D*& OutTexture
 		return false;
 	}
 
-	const FLinearColor InputMask = FGLTFMaterialUtility::GetMask(*MaterialInput);
+	const FLinearColor InputMask = FGLTFMaterialUtilities::GetMask(*MaterialInput);
 	if (AllowedMasks.Num() > 0 && !AllowedMasks.Contains(InputMask))
 	{
 		return false;
@@ -1212,7 +1212,7 @@ bool FGLTFDelayedMaterialTask::TryGetSourceTexture(const UTexture2D*& OutTexture
 			return false;
 		}
 
-		if (!FGLTFMaterialUtility::TryGetTextureCoordinateIndex(TextureParameter, OutTexCoord, OutTransform))
+		if (!FGLTFMaterialUtilities::TryGetTextureCoordinateIndex(TextureParameter, OutTexCoord, OutTransform))
 		{
 			// TODO: report error (failed to identify texture coordinate index)
 			return false;
@@ -1253,7 +1253,7 @@ bool FGLTFDelayedMaterialTask::TryGetSourceTexture(const UTexture2D*& OutTexture
 			return false;
 		}
 
-		if (!FGLTFMaterialUtility::TryGetTextureCoordinateIndex(TextureSampler, OutTexCoord, OutTransform))
+		if (!FGLTFMaterialUtilities::TryGetTextureCoordinateIndex(TextureSampler, OutTexCoord, OutTransform))
 		{
 			// TODO: report error (failed to identify texture coordinate index)
 			return false;
@@ -1401,15 +1401,15 @@ bool FGLTFDelayedMaterialTask::TryGetBakedMaterialProperty(FGLTFJsonTextureInfo&
 		return StoreBakedPropertyTexture(OutTexInfo, PropertyBakeOutput, PropertyName);
 	}
 
-	const FVector4f MaskedConstant = FVector4f(PropertyBakeOutput.ConstantValue) * FGLTFMaterialUtility::GetPropertyMask(Property);
-	if (MaskedConstant == FGLTFMaterialUtility::GetPropertyDefaultValue(Property))
+	const FVector4f MaskedConstant = FVector4f(PropertyBakeOutput.ConstantValue) * FGLTFMaterialUtilities::GetPropertyMask(Property);
+	if (MaskedConstant == FGLTFMaterialUtilities::GetPropertyDefaultValue(Property))
 	{
 		// Constant value is the same as the property's default so we can set gltf to default.
 		OutTexInfo.Index = nullptr;
 		return true;
 	}
 
-	if (FGLTFMaterialUtility::IsNormalMap(Property))
+	if (FGLTFMaterialUtilities::IsNormalMap(Property))
 	{
 		// TODO: In some cases baking normal can result in constant vector that differs slight from default (i.e 0,0,1).
 		// Yet often, when looking at such a material, it should be exactly default. Needs further investigation.
@@ -1423,7 +1423,7 @@ bool FGLTFDelayedMaterialTask::TryGetBakedMaterialProperty(FGLTFJsonTextureInfo&
 	// even though the same material when set to opaque will properly bake AmbientOcclusion to a texture.
 	// For now, create a 1x1 texture with the constant value.
 
-	FGLTFJsonTexture* Texture = FGLTFMaterialUtility::AddTexture(
+	FGLTFJsonTexture* Texture = FGLTFMaterialUtilities::AddTexture(
 		Builder,
 		PropertyBakeOutput.Pixels,
 		PropertyBakeOutput.Size,
@@ -1448,7 +1448,7 @@ FGLTFPropertyBakeOutput FGLTFDelayedMaterialTask::BakeMaterialProperty(const FGL
 	if (MeshData == nullptr)
 	{
 		FGLTFIndexArray TexCoords;
-		FGLTFMaterialUtility::GetAllTextureCoordinateIndices(Material, Property, TexCoords);
+		FGLTFMaterialUtilities::GetAllTextureCoordinateIndices(Material, Property, TexCoords);
 
 		if (TexCoords.Num() > 0)
 		{
@@ -1477,7 +1477,7 @@ FGLTFPropertyBakeOutput FGLTFDelayedMaterialTask::BakeMaterialProperty(const FGL
 
 	// TODO: add support for calculating the ideal resolution to use for baking based on connected (texture) nodes
 
-	return FGLTFMaterialUtility::BakeMaterialProperty(
+	return FGLTFMaterialUtilities::BakeMaterialProperty(
 		TextureSize,
 		Property,
 		Material,
@@ -1493,12 +1493,12 @@ bool FGLTFDelayedMaterialTask::StoreBakedPropertyTexture(FGLTFJsonTextureInfo& O
 	const TextureAddress TextureAddress = Builder.GetBakeTilingForMaterialProperty(Material, GetPropertyGroup(PropertyBakeOutput.Property));
 	const TextureFilter TextureFilter = Builder.GetBakeFilterForMaterialProperty(Material, GetPropertyGroup(PropertyBakeOutput.Property));
 
-	FGLTFJsonTexture* Texture = FGLTFMaterialUtility::AddTexture(
+	FGLTFJsonTexture* Texture = FGLTFMaterialUtilities::AddTexture(
 		Builder,
 		PropertyBakeOutput.Pixels,
 		PropertyBakeOutput.Size,
 		true, // NOTE: we can ignore alpha in everything but TryGetBaseColorAndOpacity
-		FGLTFMaterialUtility::IsNormalMap(PropertyBakeOutput.Property),
+		FGLTFMaterialUtilities::IsNormalMap(PropertyBakeOutput.Property),
 		GetBakedTextureName(PropertyName),
 		TextureAddress,
 		TextureFilter);
diff --git a/Engine/Plugins/Enterprise/GLTFExporter/Source/GLTFExporter/Private/Tasks/GLTFDelayedMeshTasks.cpp b/Engine/Plugins/Enterprise/GLTFExporter/Source/GLTFExporter/Private/Tasks/GLTFDelayedMeshTasks.cpp
--- a/Engine/Plugins/Enterprise/GLTFExporter/Source/GLTFExporter/Private/Tasks/GLTFDelayedMeshTasks.cpp
+++ b/Engine/Plugins/Enterprise/GLTFExporter/Source/GLTFExporter/Private/Tasks/GLTFDelayedMeshTasks.cpp
@@ -1,7 +1,7 @@
 // Copyright Epic Games, Inc. All Rights Reserved.
 
 #include "Tasks/GLTFDelayedMeshTasks.h"
-#include "Converters/GLTFMeshUtility.h"
+#include "Converters/GLTFMeshUtilities.h"
 #include "Converters/GLTFBufferAdapter.h"
 #include "Builders/GLTFConvertBuilder.h"
 #include "Engine/StaticMesh.h"
@@ -111,13 +111,13 @@ namespace
 
 FString FGLTFDelayedStaticMeshTask::GetName()
 {
-	return StaticMeshComponent != nullptr ? FGLTFNameUtility::GetName(StaticMeshComponent) : StaticMesh->GetName();
+	return StaticMeshComponent != nullptr ? FGLTFNameUtilities::GetName(StaticMeshComponent) : StaticMesh->GetName();
 }
 
 void FGLTFDelayedStaticMeshTask::Process()
 {
-	FGLTFMeshUtility::FullyLoad(StaticMesh);
-	JsonMesh->Name = StaticMeshComponent != nullptr ? FGLTFNameUtility::GetName(StaticMeshComponent) : StaticMesh->GetName();
+	FGLTFMeshUtilities::FullyLoad(StaticMesh);
+	JsonMesh->Name = StaticMeshComponent != nullptr ? FGLTFNameUtilities::GetName(StaticMeshComponent) : StaticMesh->GetName();
 
 	const FStaticMeshLODResources& MeshLOD = StaticMesh->GetLODForExport(LODIndex);
 	const FPositionVertexBuffer& PositionBuffer = MeshLOD.VertexBuffers.PositionVertexBuffer;
@@ -162,12 +162,12 @@ void FGLTFDelayedStaticMeshTask::Process()
 
 	ValidateVertexBuffer(Builder, &VertexBuffer, *StaticMesh->GetName());
 
-	const TArray<FStaticMaterial>& MaterialSlots = FGLTFMeshUtility::GetMaterials(StaticMesh);
+	const TArray<FStaticMaterial>& MaterialSlots = FGLTFMeshUtilities::GetMaterials(StaticMesh);
 	const int32 MaterialCount = MaterialSlots.Num();
 
 	for (int32 MaterialIndex = 0; MaterialIndex < MaterialCount; ++MaterialIndex)
 	{
-		const FGLTFIndexArray SectionIndices = FGLTFMeshUtility::GetSectionIndices(MeshLOD, MaterialIndex);
+		const FGLTFIndexArray SectionIndices = FGLTFMeshUtilities::GetSectionIndices(MeshLOD, MaterialIndex);
 		const FGLTFMeshSection* ConvertedSection = MeshSectionConverter.GetOrAdd(&MeshLOD, SectionIndices);
 
 		FGLTFJsonPrimitive& JsonPrimitive = JsonMesh->Primitives[MaterialIndex];
@@ -214,13 +214,13 @@ void FGLTFDelayedStaticMeshTask::Process()
 
 FString FGLTFDelayedSkeletalMeshTask::GetName()
 {
-	return SkeletalMeshComponent != nullptr ? FGLTFNameUtility::GetName(SkeletalMeshComponent) : SkeletalMesh->GetName();
+	return SkeletalMeshComponent != nullptr ? FGLTFNameUtilities::GetName(SkeletalMeshComponent) : SkeletalMesh->GetName();
 }
 
 void FGLTFDelayedSkeletalMeshTask::Process()
 {
-	FGLTFMeshUtility::FullyLoad(SkeletalMesh);
-	JsonMesh->Name = SkeletalMeshComponent != nullptr ? FGLTFNameUtility::GetName(SkeletalMeshComponent) : SkeletalMesh->GetName();
+	FGLTFMeshUtilities::FullyLoad(SkeletalMesh);
+	JsonMesh->Name = SkeletalMeshComponent != nullptr ? FGLTFNameUtilities::GetName(SkeletalMeshComponent) : SkeletalMesh->GetName();
 
 	const FSkeletalMeshRenderData* RenderData = SkeletalMesh->GetResourceForRendering();
 	const FSkeletalMeshLODRenderData& MeshLOD = RenderData->LODRenderData[LODIndex];
@@ -271,12 +271,12 @@ void FGLTFDelayedSkeletalMeshTask::Process()
 
 	ValidateVertexBuffer(Builder, &VertexBuffer, *SkeletalMesh->GetName());
 
-	const TArray<FSkeletalMaterial>& MaterialSlots = FGLTFMeshUtility::GetMaterials(SkeletalMesh);
+	const TArray<FSkeletalMaterial>& MaterialSlots = FGLTFMeshUtilities::GetMaterials(SkeletalMesh);
 	const uint16 MaterialCount = MaterialSlots.Num();
 
 	for (uint16 MaterialIndex = 0; MaterialIndex < MaterialCount; ++MaterialIndex)
 	{
-		const FGLTFIndexArray SectionIndices = FGLTFMeshUtility::GetSectionIndices(MeshLOD, MaterialIndex);
+		const FGLTFIndexArray SectionIndices = FGLTFMeshUtilities::GetSectionIndices(MeshLOD, MaterialIndex);
 		const FGLTFMeshSection* ConvertedSection = MeshSectionConverter.GetOrAdd(&MeshLOD, SectionIndices);
 
 		FGLTFJsonPrimitive& JsonPrimitive = JsonMesh->Primitives[MaterialIndex];
diff --git a/Engine/Plugins/Enterprise/GLTFExporter/Source/GLTFExporter/Private/Tasks/GLTFDelayedMeshTasks.h b/Engine/Plugins/Enterprise/GLTFExporter/Source/GLTFExporter/Private/Tasks/GLTFDelayedMeshTasks.h
--- a/Engine/Plugins/Enterprise/GLTFExporter/Source/GLTFExporter/Private/Tasks/GLTFDelayedMeshTasks.h
+++ b/Engine/Plugins/Enterprise/GLTFExporter/Source/GLTFExporter/Private/Tasks/GLTFDelayedMeshTasks.h
@@ -6,7 +6,7 @@
 #include "Builders/GLTFConvertBuilder.h"
 #include "Converters/GLTFMeshSectionConverters.h"
 #include "Converters/GLTFMaterialArray.h"
-#include "Converters/GLTFNameUtility.h"
+#include "Converters/GLTFNameUtilities.h"
 #include "Engine/SkeletalMesh.h"
 #include "Engine/StaticMesh.h"
 
diff --git a/Engine/Plugins/Enterprise/GLTFExporter/Source/GLTFExporter/Private/Tasks/GLTFDelayedTextureTasks.cpp b/Engine/Plugins/Enterprise/GLTFExporter/Source/GLTFExporter/Private/Tasks/GLTFDelayedTextureTasks.cpp
--- a/Engine/Plugins/Enterprise/GLTFExporter/Source/GLTFExporter/Private/Tasks/GLTFDelayedTextureTasks.cpp
+++ b/Engine/Plugins/Enterprise/GLTFExporter/Source/GLTFExporter/Private/Tasks/GLTFDelayedTextureTasks.cpp
@@ -3,7 +3,7 @@
 #include "Tasks/GLTFDelayedTextureTasks.h"
 #include "Builders/GLTFContainerBuilder.h"
 #include "Utilities/GLTFCoreUtilities.h"
-#include "Converters/GLTFTextureUtility.h"
+#include "Converters/GLTFTextureUtilities.h"
 #include "Engine/Texture2D.h"
 #include "Engine/TextureRenderTarget2D.h"
 
@@ -14,7 +14,7 @@ FString FGLTFDelayedTexture2DTask::GetName()
 
 void FGLTFDelayedTexture2DTask::Process()
 {
-	FGLTFTextureUtility::FullyLoad(Texture2D);
+	FGLTFTextureUtilities::FullyLoad(Texture2D);
 	Texture2D->GetName(JsonTexture->Name);
 
 	const bool bFromSRGB = Texture2D->SRGB;
@@ -23,14 +23,14 @@ void FGLTFDelayedTexture2DTask::Process()
 		JsonTexture->Name += bToSRGB ? TEXT("_sRGB") : TEXT("_Linear");
 	}
 
-	const FIntPoint Size = FGLTFTextureUtility::GetInGameSize(Texture2D);
-	UTextureRenderTarget2D* RenderTarget = FGLTFTextureUtility::CreateRenderTarget(Size, false);
+	const FIntPoint Size = FGLTFTextureUtilities::GetInGameSize(Texture2D);
+	UTextureRenderTarget2D* RenderTarget = FGLTFTextureUtilities::CreateRenderTarget(Size, false);
 
 	// TODO: preserve maximum image quality (avoid compression artifacts) by copying source data (and adjustments) to a temp texture
-	FGLTFTextureUtility::DrawTexture(RenderTarget, Texture2D);
+	FGLTFTextureUtilities::DrawTexture(RenderTarget, Texture2D);
 
 	TGLTFSharedArray<FColor> Pixels;
-	if (!FGLTFTextureUtility::ReadPixels(RenderTarget, *Pixels))
+	if (!FGLTFTextureUtilities::ReadPixels(RenderTarget, *Pixels))
 	{
 		Builder.LogWarning(FString::Printf(TEXT("Failed to read pixels for 2D texture %s"), *JsonTexture->Name));
 		return;
@@ -39,12 +39,12 @@ void FGLTFDelayedTexture2DTask::Process()
 	if (Builder.ExportOptions->bAdjustNormalmaps && Texture2D->IsNormalMap())
 	{
 		// TODO: add support for adjusting normals in GLTFNormalMapPreview instead
-		FGLTFTextureUtility::FlipGreenChannel(*Pixels);
+		FGLTFTextureUtilities::FlipGreenChannel(*Pixels);
 	}
 
-	FGLTFTextureUtility::TransformColorSpace(*Pixels, bFromSRGB, bToSRGB);
+	FGLTFTextureUtilities::TransformColorSpace(*Pixels, bFromSRGB, bToSRGB);
 
-	const bool bIgnoreAlpha = FGLTFTextureUtility::IsAlphaless(Texture2D->GetPixelFormat());
+	const bool bIgnoreAlpha = FGLTFTextureUtilities::IsAlphaless(Texture2D->GetPixelFormat());
 
 	JsonTexture->Source = Builder.AddUniqueImage(Pixels, Size, bIgnoreAlpha, JsonTexture->Name);
 	JsonTexture->Sampler = Builder.AddUniqueSampler(Texture2D);
@@ -57,7 +57,7 @@ FString FGLTFDelayedTextureRenderTarget2DTask::GetName()
 
 void FGLTFDelayedTextureRenderTarget2DTask::Process()
 {
-	FGLTFTextureUtility::FullyLoad(RenderTarget2D);
+	FGLTFTextureUtilities::FullyLoad(RenderTarget2D);
 	RenderTarget2D->GetName(JsonTexture->Name);
 
 	const bool bFromSRGB = RenderTarget2D->SRGB;
@@ -67,16 +67,16 @@ void FGLTFDelayedTextureRenderTarget2DTask::Process()
 	}
 
 	TGLTFSharedArray<FColor> Pixels;
-	if (!FGLTFTextureUtility::ReadPixels(RenderTarget2D, *Pixels))
+	if (!FGLTFTextureUtilities::ReadPixels(RenderTarget2D, *Pixels))
 	{
 		Builder.LogWarning(FString::Printf(TEXT("Failed to read pixels for 2D render target %s"), *JsonTexture->Name));
 		return;
 	}
 
-	FGLTFTextureUtility::TransformColorSpace(*Pixels, bFromSRGB, bToSRGB);
+	FGLTFTextureUtilities::TransformColorSpace(*Pixels, bFromSRGB, bToSRGB);
 
 	const FIntPoint Size = { RenderTarget2D->SizeX, RenderTarget2D->SizeY };
-	const bool bIgnoreAlpha = FGLTFTextureUtility::IsAlphaless(RenderTarget2D->GetFormat());
+	const bool bIgnoreAlpha = FGLTFTextureUtilities::IsAlphaless(RenderTarget2D->GetFormat());
 
 	JsonTexture->Source = Builder.AddUniqueImage(Pixels, Size, bIgnoreAlpha, JsonTexture->Name);
 	JsonTexture->Sampler = Builder.AddUniqueSampler(RenderTarget2D);
diff --git a/Engine/Plugins/Enterprise/GLTFExporter/Source/GLTFExporter/Private/Utilities/GLTFCoreUtilities.cpp b/Engine/Plugins/Enterprise/GLTFExporter/Source/GLTFExporter/Private/Utilities/GLTFCoreUtilities.cpp
--- a/Engine/Plugins/Enterprise/GLTFExporter/Source/GLTFExporter/Private/Utilities/GLTFCoreUtilities.cpp
+++ b/Engine/Plugins/Enterprise/GLTFExporter/Source/GLTFExporter/Private/Utilities/GLTFCoreUtilities.cpp
@@ -1,7 +1,7 @@
 // Copyright Epic Games, Inc. All Rights Reserved.
 
 #include "Utilities/GLTFCoreUtilities.h"
-#include "Converters/GLTFTextureUtility.h"
+#include "Converters/GLTFTextureUtilities.h"
 
 float FGLTFCoreUtilities::ConvertLength(const float Length, const float ConversionScale)
 {
