Fix support for left-handed vertex tangent basis when exporting to glTF

diff --git a/Engine/Plugins/Enterprise/GLTFExporter/Source/GLTFExporter/Private/Converters/GLTFAccessorConverters.cpp b/Engine/Plugins/Enterprise/GLTFExporter/Source/GLTFExporter/Private/Converters/GLTFAccessorConverters.cpp
--- a/Engine/Plugins/Enterprise/GLTFExporter/Source/GLTFExporter/Private/Converters/GLTFAccessorConverters.cpp
+++ b/Engine/Plugins/Enterprise/GLTFExporter/Source/GLTFExporter/Private/Converters/GLTFAccessorConverters.cpp
@@ -4,29 +4,49 @@
 #include "Utilities/GLTFCoreUtilities.h"
 #include "Converters/GLTFBufferAdapter.h"
 #include "Builders/GLTFConvertBuilder.h"
-
 #include "Rendering/ColorVertexBuffer.h"
 #include "Rendering/PositionVertexBuffer.h"
-#include "RawIndexBuffer.h"
-#include "RHI.h"
-#include "RHIResources.h"
 #include "Rendering/SkinWeightVertexBuffer.h"
 #include "Rendering/StaticMeshVertexBuffer.h"
 
-// TODO: Unreal-style implementation of std::conditional to avoid mixing in STL. Should be added to the engine.
-template <bool Condition, class TypeIfTrue, class TypeIfFalse>
-class TConditional
+namespace
 {
-public:
-	typedef TypeIfFalse Type;
-};
+	template <typename DestinationType, typename SourceType>
+	struct TGLTFVertexNormalUtilities
+	{
+		static DestinationType Convert(const SourceType& Normal)
+		{
+			return FGLTFCoreUtilities::ConvertNormal(Normal);
+		}
+	};
 
-template <class TypeIfTrue, class TypeIfFalse>
-class TConditional<true, TypeIfTrue, TypeIfFalse>
-{
-public:
-	typedef TypeIfTrue Type;
-};
+	template <typename SourceType>
+	struct TGLTFVertexNormalUtilities<FGLTFVector3, SourceType>
+	{
+		static FGLTFVector3 Convert(const SourceType& Normal)
+		{
+			return FGLTFCoreUtilities::ConvertNormal(Normal.ToFVector3f());
+		}
+	};
+
+	template <typename DestinationType, typename SourceType>
+	struct TGLTFVertexTangentUtilities
+	{
+		static DestinationType Convert(const SourceType& Tangent)
+		{
+			return FGLTFCoreUtilities::ConvertTangent(Tangent);
+		}
+	};
+
+	template <typename SourceType>
+	struct TGLTFVertexTangentUtilities<FGLTFVector4, SourceType>
+	{
+		static FGLTFVector4 Convert(const SourceType& Tangent)
+		{
+			return FGLTFCoreUtilities::ConvertTangent(Tangent.ToFVector4f());
+		}
+	};
+}
 
 FGLTFJsonAccessor* FGLTFPositionBufferConverter::Convert(const FGLTFMeshSection* MeshSection, const FPositionVertexBuffer* VertexBuffer)
 {
@@ -201,10 +221,8 @@ FGLTFJsonBufferView* FGLTFNormalBufferConverter::ConvertBufferView(const FGLTFMe
 	for (uint32 VertexIndex = 0; VertexIndex < VertexCount; ++VertexIndex)
 	{
 		const uint32 MappedVertexIndex = IndexMap[VertexIndex];
-		const FVector3f SafeNormal = TangentData[MappedVertexIndex].TangentZ.ToFVector3f().GetSafeNormal();
-
-		typedef typename TConditional<TIsSame<DestinationType, FGLTFVector3>::Value, FVector3f, SourceType>::Type IntermediateType;
-		Normals[VertexIndex] = FGLTFCoreUtilities::ConvertNormal(IntermediateType(SafeNormal));
+		const SourceType Normal = TangentData[MappedVertexIndex].TangentZ;
+		Normals[VertexIndex] = TGLTFVertexNormalUtilities<DestinationType, SourceType>::Convert(Normal);
 	}
 
 	return Builder.AddBufferView(Normals, EGLTFJsonBufferTarget::ArrayBuffer);
@@ -280,10 +298,8 @@ FGLTFJsonBufferView* FGLTFTangentBufferConverter::ConvertBufferView(const FGLTFM
 	for (uint32 VertexIndex = 0; VertexIndex < VertexCount; ++VertexIndex)
 	{
 		const uint32 MappedVertexIndex = IndexMap[VertexIndex];
-		const FVector3f SafeTangent = VertexTangents[MappedVertexIndex].TangentX.ToFVector3f().GetSafeNormal();
-
-		typedef typename TConditional<TIsSame<DestinationType, FGLTFVector4>::Value, FVector3f, SourceType>::Type IntermediateType;
-		Tangents[VertexIndex] = FGLTFCoreUtilities::ConvertTangent(IntermediateType(SafeTangent));
+		const SourceType Tangent = VertexTangents[MappedVertexIndex].TangentX;
+		Tangents[VertexIndex] = TGLTFVertexTangentUtilities<DestinationType, SourceType>::Convert(Tangent);
 	}
 
 	return Builder.AddBufferView(Tangents, EGLTFJsonBufferTarget::ArrayBuffer);
diff --git a/Engine/Plugins/Enterprise/GLTFExporter/Source/GLTFExporter/Private/Utilities/GLTFCoreUtilities.cpp b/Engine/Plugins/Enterprise/GLTFExporter/Source/GLTFExporter/Private/Utilities/GLTFCoreUtilities.cpp
--- a/Engine/Plugins/Enterprise/GLTFExporter/Source/GLTFExporter/Private/Utilities/GLTFCoreUtilities.cpp
+++ b/Engine/Plugins/Enterprise/GLTFExporter/Source/GLTFExporter/Private/Utilities/GLTFCoreUtilities.cpp
@@ -27,33 +27,38 @@ FGLTFVector3 FGLTFCoreUtilities::ConvertScale(const FVector3f& Scale)
 
 FGLTFVector3 FGLTFCoreUtilities::ConvertNormal(const FVector3f& Normal)
 {
-	return ConvertVector(Normal);
+	const FVector3f SafeNormal = Normal.GetSafeNormal();
+	return ConvertVector(SafeNormal);
 }
 
-FGLTFInt16Vector4 FGLTFCoreUtilities::ConvertNormal(const FPackedRGBA16N& Normal)
+FGLTFInt8Vector4 FGLTFCoreUtilities::ConvertNormal(const FPackedNormal& Normal)
 {
-	return { Normal.X, Normal.Z, Normal.Y, 0 };
+	const FPackedNormal SafeNormal = Normal.ToFVector3f().GetSafeNormal();
+	return { SafeNormal.Vector.X, SafeNormal.Vector.Z, SafeNormal.Vector.Y, 0 };
 }
 
-FGLTFInt8Vector4 FGLTFCoreUtilities::ConvertNormal(const FPackedNormal& Normal)
+FGLTFInt16Vector4 FGLTFCoreUtilities::ConvertNormal(const FPackedRGBA16N& Normal)
 {
-	return { Normal.Vector.X, Normal.Vector.Z, Normal.Vector.Y, 0 };
+	const FPackedRGBA16N SafeNormal = Normal.ToFVector3f().GetSafeNormal();
+	return { SafeNormal.X, SafeNormal.Z, SafeNormal.Y, 0 };
 }
 
-FGLTFVector4 FGLTFCoreUtilities::ConvertTangent(const FVector3f& Tangent)
+FGLTFVector4 FGLTFCoreUtilities::ConvertTangent(const FVector4f& Tangent)
 {
-	// glTF stores tangent as Vec4, with W component indicating handedness of tangent basis.
-	return { Tangent.X, Tangent.Z, Tangent.Y, 1.0f };
+	const FVector3f SafeTangent = Tangent.GetSafeNormal();
+	return { SafeTangent.X, SafeTangent.Z, SafeTangent.Y, Tangent.W < 0 ? -1.0f : +1.0f };
 }
 
-FGLTFInt16Vector4 FGLTFCoreUtilities::ConvertTangent(const FPackedRGBA16N& Tangent)
+FGLTFInt8Vector4 FGLTFCoreUtilities::ConvertTangent(const FPackedNormal& Tangent)
 {
-	return { Tangent.X, Tangent.Z, Tangent.Y, MAX_int16 /* = 1.0 */ };
+	const FPackedNormal SafeTangent = Tangent.ToFVector3f().GetSafeNormal();
+	return { SafeTangent.Vector.X, SafeTangent.Vector.Z, SafeTangent.Vector.Y, Tangent.Vector.W < 0 ? MIN_int8 : MAX_int8 };
 }
 
-FGLTFInt8Vector4 FGLTFCoreUtilities::ConvertTangent(const FPackedNormal& Tangent)
+FGLTFInt16Vector4 FGLTFCoreUtilities::ConvertTangent(const FPackedRGBA16N& Tangent)
 {
-	return { Tangent.Vector.X, Tangent.Vector.Z, Tangent.Vector.Y, MAX_int8 /* = 1.0 */ };
+	const FPackedRGBA16N SafeTangent = Tangent.ToFVector3f().GetSafeNormal();
+	return { SafeTangent.X, SafeTangent.Z, SafeTangent.Y, Tangent.W < 0 ? MIN_int16 : MAX_int16 };
 }
 
 FGLTFVector2 FGLTFCoreUtilities::ConvertUV(const FVector2f& UV)
diff --git a/Engine/Plugins/Enterprise/GLTFExporter/Source/GLTFExporter/Public/Utilities/GLTFCoreUtilities.h b/Engine/Plugins/Enterprise/GLTFExporter/Source/GLTFExporter/Public/Utilities/GLTFCoreUtilities.h
--- a/Engine/Plugins/Enterprise/GLTFExporter/Source/GLTFExporter/Public/Utilities/GLTFCoreUtilities.h
+++ b/Engine/Plugins/Enterprise/GLTFExporter/Source/GLTFExporter/Public/Utilities/GLTFCoreUtilities.h
@@ -26,12 +26,12 @@ struct GLTFEXPORTER_API FGLTFCoreUtilities
 	static FGLTFVector3 ConvertScale(const FVector3f& Scale);
 
 	static FGLTFVector3 ConvertNormal(const FVector3f& Normal);
-	static FGLTFInt16Vector4 ConvertNormal(const FPackedRGBA16N& Normal);
 	static FGLTFInt8Vector4 ConvertNormal(const FPackedNormal& Normal);
+	static FGLTFInt16Vector4 ConvertNormal(const FPackedRGBA16N& Normal);
 
-	static FGLTFVector4 ConvertTangent(const FVector3f& Tangent);
-	static FGLTFInt16Vector4 ConvertTangent(const FPackedRGBA16N& Tangent);
+	static FGLTFVector4 ConvertTangent(const FVector4f& Tangent);
 	static FGLTFInt8Vector4 ConvertTangent(const FPackedNormal& Tangent);
+	static FGLTFInt16Vector4 ConvertTangent(const FPackedRGBA16N& Tangent);
 
 	static FGLTFVector2 ConvertUV(const FVector2f& UV);
 	static FGLTFVector2 ConvertUV(const FVector2DHalf& UV);
