Fix issues when material baking the following expressions for gltf export:
- TransformPosition & TransformVector (to and from local space)
- ObjectScale
- LocalPosition

diff --git a/Engine/Plugins/Enterprise/GLTFExporter/Source/GLTFExporter/Private/Converters/GLTFMeshData.cpp b/Engine/Plugins/Enterprise/GLTFExporter/Source/GLTFExporter/Private/Converters/GLTFMeshData.cpp
--- a/Engine/Plugins/Enterprise/GLTFExporter/Source/GLTFExporter/Private/Converters/GLTFMeshData.cpp
+++ b/Engine/Plugins/Enterprise/GLTFExporter/Source/GLTFExporter/Private/Converters/GLTFMeshData.cpp
@@ -2,7 +2,6 @@
 
 #include "Converters/GLTFMeshData.h"
 #include "Converters/GLTFNameUtilities.h"
-
 #include "Components/SkeletalMeshComponent.h"
 #include "Components/StaticMeshComponent.h"
 #include "Engine/StaticMesh.h"
@@ -16,6 +15,83 @@
 #include "StaticMeshAttributes.h"
 #endif
 
+#if WITH_EDITOR
+#include "StaticMeshOperations.h"
+#include "Components/SplineMeshComponent.h"
+
+namespace
+{
+	// TODO: replace hack with FMeshMergeHelpers::RetrieveMesh after support for bApplyComponentTransform is approved and merged
+	void RetrieveMeshHack(const UStaticMeshComponent* StaticMeshComponent, int32 LODIndex, FMeshDescription& OutMeshDescription, bool bPropagateVertexColours, bool bApplyComponentTransform)
+	{
+		const UStaticMesh* StaticMesh = StaticMeshComponent->GetStaticMesh();
+		check(StaticMesh);
+
+		// Export the mesh data using static mesh render data
+		FMeshMergeHelpers::ExportStaticMeshLOD(StaticMesh->GetRenderData()->LODResources[LODIndex], OutMeshDescription, StaticMesh->GetStaticMaterials());
+
+		// Make sure the mesh is not irreparably malformed.
+		if (OutMeshDescription.VertexInstances().Num() <= 0)
+		{
+			return;
+		}
+
+		// If we have a component, use it to retrieve transform & vertex colors (if requested)
+		if (StaticMeshComponent)
+		{
+			// Handle spline mesh deformation
+			const bool bIsSplineMeshComponent = StaticMeshComponent->IsA<USplineMeshComponent>();
+			if (bIsSplineMeshComponent)
+			{
+				const USplineMeshComponent* SplineMeshComponent = Cast<USplineMeshComponent>(StaticMeshComponent);
+				// Deform mesh data according to the Spline Mesh Component's data
+				FMeshMergeHelpers::PropagateSplineDeformationToMesh(SplineMeshComponent, OutMeshDescription);
+			}
+
+			// If specified propagate painted vertex colors into our raw mesh
+			if (bPropagateVertexColours)
+			{
+				FMeshMergeHelpers::PropagatePaintedColorsToMesh(StaticMeshComponent, LODIndex, OutMeshDescription);
+			}
+
+			// If specified transform vertices from local to world space
+			if (bApplyComponentTransform)
+			{
+				FStaticMeshOperations::ApplyTransform(OutMeshDescription, StaticMeshComponent->GetComponentTransform());
+			}
+		}
+
+		FMeshBuildSettings BuildSettings;
+		
+		// If editor data is not available, we won't have access to source model
+		const bool bHasSourceModels = StaticMesh->IsSourceModelValid(0);
+		if (bHasSourceModels)
+		{ 
+			// Imported meshes will have a valid mesh description
+			const bool bImportedMesh = bHasSourceModels ? false : StaticMesh->IsMeshDescriptionValid(LODIndex);
+
+			// Use build settings from base mesh for LOD entries that were generated inside Editor.
+			BuildSettings = bImportedMesh ? StaticMesh->GetSourceModel(LODIndex).BuildSettings : StaticMesh->GetSourceModel(0).BuildSettings;
+		}
+
+		// Figure out if we should recompute normals and tangents. By default generated LODs should not recompute normals
+		EComputeNTBsFlags ComputeNTBsOptions = EComputeNTBsFlags::BlendOverlappingNormals;
+		if (BuildSettings.bRemoveDegenerates)
+		{
+			// If removing degenerate triangles, ignore them when computing tangents.
+			ComputeNTBsOptions |= EComputeNTBsFlags::IgnoreDegenerateTriangles;
+		}
+		if (BuildSettings.bUseMikkTSpace)
+		{
+			ComputeNTBsOptions |= EComputeNTBsFlags::UseMikkTSpace;
+		}
+
+		FStaticMeshOperations::ComputeTriangleTangentsAndNormals(OutMeshDescription, 0.0f);
+		FStaticMeshOperations::RecomputeNormalsAndTangentsIfNeeded(OutMeshDescription, ComputeNTBsOptions);
+	}
+}
+#endif
+
 FGLTFMeshData::FGLTFMeshData(const UStaticMesh* StaticMesh, const UStaticMeshComponent* StaticMeshComponent, int32 LODIndex)
 	: Parent(nullptr)
 	, LODIndex(LODIndex)
@@ -35,7 +111,7 @@ FGLTFMeshData::FGLTFMeshData(const UStaticMesh* StaticMesh, const UStaticMeshCom
 
 #if WITH_EDITOR
 		PrimitiveData = { StaticMeshComponent };
-		FMeshMergeHelpers::RetrieveMesh(StaticMeshComponent, LODIndex, Description, true);
+		RetrieveMeshHack(StaticMeshComponent, LODIndex, Description, true, false);
 
 		constexpr int32 LightMapLODIndex = 0; // TODO: why is this zero?
 		if (StaticMeshComponent->LODData.IsValidIndex(LightMapLODIndex))
@@ -109,15 +185,15 @@ FGLTFMeshData::FGLTFMeshData(const USkeletalMesh* SkeletalMesh, const USkeletalM
 			SpawnParams.ObjectFlags |= RF_Transient;
 			SpawnParams.bAllowDuringConstructionScript = true;
 
-			if (AActor* Actor = World->SpawnActor<AActor>(SpawnParams))
+			if (AActor* TempActor = World->SpawnActor<AActor>(SpawnParams))
 			{
-				USkeletalMeshComponent* Component = NewObject<USkeletalMeshComponent>(Actor, TEXT(""), RF_Transient);
-				Component->RegisterComponent();
-				Component->SetSkeletalMesh(const_cast<USkeletalMesh*>(SkeletalMesh));
+				USkeletalMeshComponent* TempComponent = NewObject<USkeletalMeshComponent>(TempActor, TEXT(""), RF_Transient);
+				TempComponent->RegisterComponent();
+				TempComponent->SetSkeletalMesh(const_cast<USkeletalMesh*>(SkeletalMesh));
 
-				FMeshMergeHelpers::RetrieveMesh(Component, LODIndex, Description, true);
+				FMeshMergeHelpers::RetrieveMesh(TempComponent, LODIndex, Description, true);
 
-				World->DestroyActor(Actor, false, false);
+				World->DestroyActor(TempActor, false, false);
 			}
 		}
 #endif
diff --git a/Engine/Plugins/Enterprise/GLTFExporter/Source/GLTFMaterialBaking/Private/GLTFExportMaterialProxy.h b/Engine/Plugins/Enterprise/GLTFExporter/Source/GLTFMaterialBaking/Private/GLTFExportMaterialProxy.h
--- a/Engine/Plugins/Enterprise/GLTFExporter/Source/GLTFMaterialBaking/Private/GLTFExportMaterialProxy.h
+++ b/Engine/Plugins/Enterprise/GLTFExporter/Source/GLTFMaterialBaking/Private/GLTFExportMaterialProxy.h
@@ -51,7 +51,7 @@ struct FGLTFExportMaterialCompiler : public FProxyMaterialCompiler
 	virtual int32 WorldPosition(EWorldPositionIncludedOffsets WorldPositionIncludedOffsets) override
 	{
 #if WITH_EDITOR
-		return Compiler->MaterialBakingWorldPosition();
+		return MaterialBakingWorldPosition();
 #else
 		return Compiler->WorldPosition(WorldPositionIncludedOffsets);
 #endif
@@ -148,7 +148,9 @@ struct FGLTFExportMaterialCompiler : public FProxyMaterialCompiler
 #if WITH_EDITOR
 	virtual int32 MaterialBakingWorldPosition() override
 	{
-		return Compiler->MaterialBakingWorldPosition();
+		// Depending on how the mesh data was retrieved, baking position may only be in local-space
+		const int32 BakingPosition = Compiler->MaterialBakingWorldPosition();
+		return Compiler->TransformPosition(MCB_Local, MCB_World, BakingPosition);
 	}
 #endif
 
diff --git a/Engine/Plugins/Enterprise/GLTFExporter/Source/GLTFMaterialBaking/Private/MaterialRenderItem.cpp b/Engine/Plugins/Enterprise/GLTFExporter/Source/GLTFMaterialBaking/Private/MaterialRenderItem.cpp
--- a/Engine/Plugins/Enterprise/GLTFExporter/Source/GLTFMaterialBaking/Private/MaterialRenderItem.cpp
+++ b/Engine/Plugins/Enterprise/GLTFExporter/Source/GLTFMaterialBaking/Private/MaterialRenderItem.cpp
@@ -134,20 +134,16 @@ void FGLTFMeshMaterialRenderItem::QueueMaterial(FCanvasRenderContext& RenderCont
 		{
 			FPrimitiveUniformShaderParameters PrimitiveParams = FPrimitiveUniformShaderParametersBuilder{}
 				.Defaults()
-					.LocalToWorld(FMatrix::Identity)
+					.LocalToWorld(PrimitiveData->LocalToWorld)
 					.ActorWorldPosition(PrimitiveData->ActorPosition)
-					.WorldBounds(PrimitiveData->ObjectBounds)
-					.LocalBounds(PrimitiveData->ObjectLocalBounds)
+					.WorldBounds(PrimitiveData->WorldBounds)
+					.LocalBounds(PrimitiveData->LocalBounds)
 					.PreSkinnedLocalBounds(PrimitiveData->PreSkinnedLocalBounds)
 					.CustomPrimitiveData(PrimitiveData->CustomPrimitiveData)
 					.ReceivesDecals(false)
 					//.DrawsVelocity(false)
 				.Build();
 
-			// Overwrite object orientation, since original calculation is derived from LocalToWorld matrix,
-			// which currently needs to be identity matrix to avoid transforming the "2D" vertices used for baking
-			PrimitiveParams.ObjectOrientation = FVector3f(PrimitiveData->ObjectOrientation);
-
 			const_cast<TUniformBuffer<FPrimitiveUniformShaderParameters>*>(MeshElement.Elements[0].PrimitiveUniformBufferResource)->SetContents(PrimitiveParams);
 		}
 
@@ -177,6 +173,9 @@ void FGLTFMeshMaterialRenderItem::QueueMaterial(FCanvasRenderContext& RenderCont
 
 void FGLTFMeshMaterialRenderItem::PopulateWithQuadData()
 {
+	// Pre-transform all vertices with the inverse of LocalToWorld to negate its effect during material baking 
+	const FMatrix44f WorldToLocal = (MeshSettings->PrimitiveData != nullptr) ? FMatrix44f(MeshSettings->PrimitiveData->LocalToWorld.Inverse()) : FMatrix44f::Identity;
+	
 	Vertices.Empty(4);
 	Indices.Empty(6);
 
@@ -193,8 +192,11 @@ void FGLTFMeshMaterialRenderItem::PopulateWithQuadData()
 		FDynamicMeshVertex* Vert = new(Vertices)FDynamicMeshVertex();
 		const int32 X = VertIndex & 1;
 		const int32 Y = (VertIndex >> 1) & 1;
-		Vert->Position.Set(ScaleX * X, ScaleY * Y, 0);
-		Vert->SetTangents(FVector3f(1, 0, 0), FVector3f(0, 1, 0), FVector3f(0, 0, 1));
+		Vert->Position = WorldToLocal.TransformPosition(FVector3f(ScaleX * X, ScaleY * Y, 0));
+		FVector3f TangentX = WorldToLocal.TransformVector(FVector3f(1, 0, 0));
+		FVector3f TangentZ = WorldToLocal.TransformVector(FVector3f(0, 1, 0));
+		FVector3f TangentY = WorldToLocal.TransformVector(FVector3f(0, 0, 1));
+		Vert->SetTangents(TangentX, TangentZ, TangentY);
 		FMemory::Memzero(&Vert->TextureCoordinate, sizeof(Vert->TextureCoordinate));
 		for (int32 TexcoordIndex = 0; TexcoordIndex < MAX_STATIC_TEXCOORDS; TexcoordIndex++)
 		{
@@ -210,6 +212,9 @@ void FGLTFMeshMaterialRenderItem::PopulateWithQuadData()
 
 void FGLTFMeshMaterialRenderItem::PopulateWithMeshData()
 {
+	// Pre-transform all vertices with the inverse of LocalToWorld to negate its effect during material baking 
+	const FMatrix44f WorldToLocal = (MeshSettings->PrimitiveData != nullptr) ? FMatrix44f(MeshSettings->PrimitiveData->LocalToWorld.Inverse()) : FMatrix44f::Identity;
+
 	const FMeshDescription* RawMesh = MeshSettings->MeshDescription;
 
 	FStaticMeshConstAttributes Attributes(*RawMesh);
@@ -267,15 +272,15 @@ void FGLTFMeshMaterialRenderItem::PopulateWithMeshData()
 				{
 					// compute vertex position from original UV
 					const FVector2D& UV = FVector2D(VertexInstanceUVs.Get(SrcVertexInstanceID, MeshSettings->TextureCoordinateIndex));
-					Vert->Position.Set(UV.X * ScaleX, UV.Y * ScaleY, 0);
+					Vert->Position = WorldToLocal.TransformPosition(FVector3f(UV.X * ScaleX, UV.Y * ScaleY, 0));
 				}
 				else
 				{
 					const FVector2D& UV = MeshSettings->CustomTextureCoordinates[SrcVertIndex];
-					Vert->Position.Set(UV.X * ScaleX, UV.Y * ScaleY, 0);
+					Vert->Position = WorldToLocal.TransformPosition(FVector3f(UV.X * ScaleX, UV.Y * ScaleY, 0));
 				}
-				FVector3f TangentX = VertexInstanceTangents[SrcVertexInstanceID];
-				FVector3f TangentZ = VertexInstanceNormals[SrcVertexInstanceID];
+				FVector3f TangentX = WorldToLocal.TransformVector(VertexInstanceTangents[SrcVertexInstanceID]);
+				FVector3f TangentZ = WorldToLocal.TransformVector(VertexInstanceNormals[SrcVertexInstanceID]);
 				FVector3f TangentY = FVector3f::CrossProduct(TangentZ, TangentX).GetSafeNormal() * VertexInstanceBinormalSigns[SrcVertexInstanceID];
 				Vert->SetTangents(TangentX, TangentY, TangentZ);
 				for (int32 TexcoordIndex = 0; TexcoordIndex < NumTexcoords; TexcoordIndex++)
diff --git a/Engine/Plugins/Enterprise/GLTFExporter/Source/GLTFMaterialBaking/Public/GLTFMaterialBakingStructures.h b/Engine/Plugins/Enterprise/GLTFExporter/Source/GLTFMaterialBaking/Public/GLTFMaterialBakingStructures.h
--- a/Engine/Plugins/Enterprise/GLTFExporter/Source/GLTFMaterialBaking/Public/GLTFMaterialBakingStructures.h
+++ b/Engine/Plugins/Enterprise/GLTFExporter/Source/GLTFMaterialBaking/Public/GLTFMaterialBakingStructures.h
@@ -69,10 +69,10 @@ struct FGLTFMaterialDataEx
 struct FGLTFPrimitiveData
 {
 	FGLTFPrimitiveData(const FBoxSphereBounds& LocalBounds = FBoxSphereBounds(ForceInitToZero))
-		: ActorPosition(ForceInitToZero)
-		, ObjectOrientation(FVector::UpVector)
-		, ObjectBounds(LocalBounds)
-		, ObjectLocalBounds(LocalBounds)
+		: LocalToWorld(FMatrix::Identity)
+		, ActorPosition(ForceInitToZero)
+		, WorldBounds(LocalBounds)
+		, LocalBounds(LocalBounds)
 		, PreSkinnedLocalBounds(LocalBounds)
 		, CustomPrimitiveData(nullptr)
 	{}
@@ -86,10 +86,10 @@ struct FGLTFPrimitiveData
 	{}
 
 	FGLTFPrimitiveData(const UMeshComponent* MeshComponent)
-		: ActorPosition(MeshComponent->GetComponentLocation())
-		, ObjectOrientation(MeshComponent->GetRenderMatrix().GetUnitAxis(EAxis::Z))
-		, ObjectBounds(MeshComponent->CalcBounds(MeshComponent->GetComponentTransform()))
-		, ObjectLocalBounds(MeshComponent->CalcLocalBounds())
+		: LocalToWorld(MeshComponent->GetRenderMatrix())
+		, ActorPosition(MeshComponent->GetActorPositionForRenderer())
+		, WorldBounds(MeshComponent->Bounds)
+		, LocalBounds(MeshComponent->GetLocalBounds())
 		, CustomPrimitiveData(&MeshComponent->GetCustomPrimitiveData())
 	{
 		if (const USkinnedMeshComponent* SkinnedMeshComponent = Cast<USkinnedMeshComponent>(MeshComponent))
@@ -98,21 +98,21 @@ struct FGLTFPrimitiveData
 		}
 		else
 		{
-			PreSkinnedLocalBounds = ObjectLocalBounds;
+			PreSkinnedLocalBounds = LocalBounds;
 		}
 	}
 
-	/** The location of the actor in world-space */
+	/** The mesh component's local-to-world transform */
+	FMatrix LocalToWorld;
+	/** The actor's location in world-space */
 	FVector ActorPosition;
-	/** The up vector of the object in world-space */
-	FVector ObjectOrientation;
-	/** The bounds of the object in world-space */
-	FBoxSphereBounds ObjectBounds;
-	/** The bounds of the object in local-space */
-	FBoxSphereBounds ObjectLocalBounds;
-	/** The pre-skinning bounds for the object in local-space */
+	/** The mesh component's bounds in world-space */
+	FBoxSphereBounds WorldBounds;
+	/** The mesh component's bounds in local-space */
+	FBoxSphereBounds LocalBounds;
+	/** The mesh component's pre-skinning bounds in local-space */
 	FBoxSphereBounds PreSkinnedLocalBounds;
-	/** The custom primitive data for the mesh component */
+	/** The mesh component's custom primitive data */
 	const FCustomPrimitiveData* CustomPrimitiveData;
 };
 
