Remove workaround for MeshMergeHelpers in glTF exporter which also fixes horde issue 246908

diff --git a/Engine/Plugins/Enterprise/GLTFExporter/Source/GLTFExporter/Private/Converters/GLTFMeshData.cpp b/Engine/Plugins/Enterprise/GLTFExporter/Source/GLTFExporter/Private/Converters/GLTFMeshData.cpp
--- a/Engine/Plugins/Enterprise/GLTFExporter/Source/GLTFExporter/Private/Converters/GLTFMeshData.cpp
+++ b/Engine/Plugins/Enterprise/GLTFExporter/Source/GLTFExporter/Private/Converters/GLTFMeshData.cpp
@@ -15,83 +15,6 @@
 #include "StaticMeshAttributes.h"
 #endif
 
-#if WITH_EDITOR
-#include "StaticMeshOperations.h"
-#include "Components/SplineMeshComponent.h"
-
-namespace
-{
-	// TODO: replace hack with FMeshMergeHelpers::RetrieveMesh after support for bApplyComponentTransform is approved and merged
-	void RetrieveMeshHack(const UStaticMeshComponent* StaticMeshComponent, int32 LODIndex, FMeshDescription& OutMeshDescription, bool bPropagateVertexColours, bool bApplyComponentTransform)
-	{
-		const UStaticMesh* StaticMesh = StaticMeshComponent->GetStaticMesh();
-		check(StaticMesh);
-
-		// Export the mesh data using static mesh render data
-		FMeshMergeHelpers::ExportStaticMeshLOD(StaticMesh->GetRenderData()->LODResources[LODIndex], OutMeshDescription, StaticMesh->GetStaticMaterials());
-
-		// Make sure the mesh is not irreparably malformed.
-		if (OutMeshDescription.VertexInstances().Num() <= 0)
-		{
-			return;
-		}
-
-		// If we have a component, use it to retrieve transform & vertex colors (if requested)
-		if (StaticMeshComponent)
-		{
-			// Handle spline mesh deformation
-			const bool bIsSplineMeshComponent = StaticMeshComponent->IsA<USplineMeshComponent>();
-			if (bIsSplineMeshComponent)
-			{
-				const USplineMeshComponent* SplineMeshComponent = Cast<USplineMeshComponent>(StaticMeshComponent);
-				// Deform mesh data according to the Spline Mesh Component's data
-				FMeshMergeHelpers::PropagateSplineDeformationToMesh(SplineMeshComponent, OutMeshDescription);
-			}
-
-			// If specified propagate painted vertex colors into our raw mesh
-			if (bPropagateVertexColours)
-			{
-				FMeshMergeHelpers::PropagatePaintedColorsToMesh(StaticMeshComponent, LODIndex, OutMeshDescription);
-			}
-
-			// If specified transform vertices from local to world space
-			if (bApplyComponentTransform)
-			{
-				FStaticMeshOperations::ApplyTransform(OutMeshDescription, StaticMeshComponent->GetComponentTransform());
-			}
-		}
-
-		FMeshBuildSettings BuildSettings;
-		
-		// If editor data is not available, we won't have access to source model
-		const bool bHasSourceModels = StaticMesh->IsSourceModelValid(0);
-		if (bHasSourceModels)
-		{ 
-			// Imported meshes will have a valid mesh description
-			const bool bImportedMesh = bHasSourceModels ? false : StaticMesh->IsMeshDescriptionValid(LODIndex);
-
-			// Use build settings from base mesh for LOD entries that were generated inside Editor.
-			BuildSettings = bImportedMesh ? StaticMesh->GetSourceModel(LODIndex).BuildSettings : StaticMesh->GetSourceModel(0).BuildSettings;
-		}
-
-		// Figure out if we should recompute normals and tangents. By default generated LODs should not recompute normals
-		EComputeNTBsFlags ComputeNTBsOptions = EComputeNTBsFlags::BlendOverlappingNormals;
-		if (BuildSettings.bRemoveDegenerates)
-		{
-			// If removing degenerate triangles, ignore them when computing tangents.
-			ComputeNTBsOptions |= EComputeNTBsFlags::IgnoreDegenerateTriangles;
-		}
-		if (BuildSettings.bUseMikkTSpace)
-		{
-			ComputeNTBsOptions |= EComputeNTBsFlags::UseMikkTSpace;
-		}
-
-		FStaticMeshOperations::ComputeTriangleTangentsAndNormals(OutMeshDescription, 0.0f);
-		FStaticMeshOperations::RecomputeNormalsAndTangentsIfNeeded(OutMeshDescription, ComputeNTBsOptions);
-	}
-}
-#endif
-
 FGLTFMeshData::FGLTFMeshData(const UStaticMesh* StaticMesh, const UStaticMeshComponent* StaticMeshComponent, int32 LODIndex)
 	: Parent(nullptr)
 	, LODIndex(LODIndex)
@@ -111,7 +34,7 @@ FGLTFMeshData::FGLTFMeshData(const UStaticMesh* StaticMesh, const UStaticMeshCom
 
 #if WITH_EDITOR
 		PrimitiveData = { StaticMeshComponent };
-		RetrieveMeshHack(StaticMeshComponent, LODIndex, Description, true, false);
+		FMeshMergeHelpers::RetrieveMesh(StaticMeshComponent, LODIndex, Description, true, false);
 
 		constexpr int32 LightMapLODIndex = 0; // TODO: why is this zero?
 		if (StaticMeshComponent->LODData.IsValidIndex(LightMapLODIndex))
@@ -163,7 +86,7 @@ FGLTFMeshData::FGLTFMeshData(const USkeletalMesh* SkeletalMesh, const USkeletalM
 
 #if WITH_EDITOR
 		PrimitiveData = { SkeletalMeshComponent };
-		FMeshMergeHelpers::RetrieveMesh(SkeletalMeshComponent, LODIndex, Description, true);
+		FMeshMergeHelpers::RetrieveMesh(SkeletalMeshComponent, LODIndex, Description, true, false);
 #endif
 	}
 	else
@@ -191,7 +114,7 @@ FGLTFMeshData::FGLTFMeshData(const USkeletalMesh* SkeletalMesh, const USkeletalM
 				TempComponent->RegisterComponent();
 				TempComponent->SetSkeletalMesh(const_cast<USkeletalMesh*>(SkeletalMesh));
 
-				FMeshMergeHelpers::RetrieveMesh(TempComponent, LODIndex, Description, true);
+				FMeshMergeHelpers::RetrieveMesh(TempComponent, LODIndex, Description, true, false);
 
 				World->DestroyActor(TempActor, false, false);
 			}
