Minor code improvement in GLTFAccessorConverters

diff --git a/Engine/Plugins/Enterprise/GLTFExporter/Source/GLTFExporter/Private/Converters/GLTFAccessorConverters.cpp b/Engine/Plugins/Enterprise/GLTFExporter/Source/GLTFExporter/Private/Converters/GLTFAccessorConverters.cpp
--- a/Engine/Plugins/Enterprise/GLTFExporter/Source/GLTFExporter/Private/Converters/GLTFAccessorConverters.cpp
+++ b/Engine/Plugins/Enterprise/GLTFExporter/Source/GLTFExporter/Private/Converters/GLTFAccessorConverters.cpp
@@ -221,7 +221,7 @@ FGLTFJsonBufferView* FGLTFNormalBufferConverter::ConvertBufferView(const FGLTFMe
 	for (uint32 VertexIndex = 0; VertexIndex < VertexCount; ++VertexIndex)
 	{
 		const uint32 MappedVertexIndex = IndexMap[VertexIndex];
-		const SourceType Normal = TangentData[MappedVertexIndex].TangentZ;
+		const SourceType& Normal = TangentData[MappedVertexIndex].TangentZ;
 		Normals[VertexIndex] = TGLTFVertexNormalUtilities<DestinationType, SourceType>::Convert(Normal);
 	}
 
@@ -298,7 +298,7 @@ FGLTFJsonBufferView* FGLTFTangentBufferConverter::ConvertBufferView(const FGLTFM
 	for (uint32 VertexIndex = 0; VertexIndex < VertexCount; ++VertexIndex)
 	{
 		const uint32 MappedVertexIndex = IndexMap[VertexIndex];
-		const SourceType Tangent = VertexTangents[MappedVertexIndex].TangentX;
+		const SourceType& Tangent = VertexTangents[MappedVertexIndex].TangentX;
 		Tangents[VertexIndex] = TGLTFVertexTangentUtilities<DestinationType, SourceType>::Convert(Tangent);
 	}
 
