Update FilterPlugin.ini in GLTFExporter to reflect removed documentation

diff --git a/Engine/Plugins/Enterprise/GLTFExporter/Config/FilterPlugin.ini b/Engine/Plugins/Enterprise/GLTFExporter/Config/FilterPlugin.ini
--- a/Engine/Plugins/Enterprise/GLTFExporter/Config/FilterPlugin.ini
+++ b/Engine/Plugins/Enterprise/GLTFExporter/Config/FilterPlugin.ini
@@ -1,4 +1,2 @@
 [FilterPlugin]
-/Docs/
-/README.md
 /Shaders/*
