Fixed issue in gltf exporter baking module causing it to incorrectly skip baking some properties based on the material's blend mode.

diff --git a/Engine/Plugins/Enterprise/GLTFExporter/Source/GLTFMaterialBaking/Private/GLTFExportMaterialProxy.h b/Engine/Plugins/Enterprise/GLTFExporter/Source/GLTFMaterialBaking/Private/GLTFExportMaterialProxy.h
--- a/Engine/Plugins/Enterprise/GLTFExporter/Source/GLTFMaterialBaking/Private/GLTFExportMaterialProxy.h
+++ b/Engine/Plugins/Enterprise/GLTFExporter/Source/GLTFMaterialBaking/Private/GLTFExportMaterialProxy.h
@@ -350,57 +350,38 @@ public:
 	{
 		if (Property == MP_EmissiveColor)
 		{
-			const EBlendMode BlendMode = MaterialInterface->GetBlendMode();
 			FGLTFExportMaterialCompiler ProxyCompiler(Compiler);
 			const uint32 ForceCast_Exact_Replicate = MFCF_ForceCast | MFCF_ExactMatch | MFCF_ReplicateValue;
 
 			switch (PropertyToCompile)
 			{
 			case MP_EmissiveColor:
-				// Emissive is ALWAYS returned...
-				return MaterialInterface->CompileProperty(&ProxyCompiler, MP_EmissiveColor, ForceCast_Exact_Replicate);
 			case MP_BaseColor:
-				return MaterialInterface->CompileProperty(&ProxyCompiler, MP_BaseColor, ForceCast_Exact_Replicate);
-				break;
 			case MP_Specular:
 			case MP_Roughness:
 			case MP_Anisotropy:
 			case MP_Metallic:
 			case MP_AmbientOcclusion:
-				// Only return for Opaque and Masked...
-				if (BlendMode == BLEND_Opaque || BlendMode == BLEND_Masked)
-				{
-					return MaterialInterface->CompileProperty(&ProxyCompiler, PropertyToCompile, ForceCast_Exact_Replicate);
-				}
-				break;
-
 			case MP_Opacity:
 			case MP_OpacityMask:
 			case MP_CustomData0:
 			case MP_CustomData1:
 			case MP_SubsurfaceColor:
-			{
 				return MaterialInterface->CompileProperty(&ProxyCompiler, PropertyToCompile, ForceCast_Exact_Replicate);
-			}
 			case MP_Normal:
 			case MP_Tangent:
-				// Only return for Opaque and Masked...
-				if (BlendMode == BLEND_Opaque || BlendMode == BLEND_Masked)
+				return CompileNormalEncoding(
+					Compiler,
+					CompileNormalTransform(&ProxyCompiler, MaterialInterface->CompileProperty(&ProxyCompiler, PropertyToCompile, ForceCast_Exact_Replicate)));
+			case MP_Refraction:
+				// Only index of refraction can be supported because other methods don't have values within a suitable range for encoding into 8-bit baked textures
+				if (Material->RefractionMethod == RM_IndexOfRefraction)
 				{
-					return CompileNormalEncoding(
+					return CompileRefractionEncoding(
 						Compiler,
-						CompileNormalTransform(&ProxyCompiler, MaterialInterface->CompileProperty(&ProxyCompiler, PropertyToCompile, ForceCast_Exact_Replicate)));
+						MaterialInterface->CompileProperty(&ProxyCompiler, MP_Refraction, ForceCast_Exact_Replicate));
 				}
 				break;
-			case MP_Refraction:
-					// Only return for Translucent, Additive, and AlphaComposite...
-					if (BlendMode == BLEND_Translucent || BlendMode == BLEND_Additive || BlendMode == BLEND_AlphaComposite)
-					{
-						return CompileRefractionEncoding(
-							Compiler,
-							MaterialInterface->CompileProperty(&ProxyCompiler, MP_Refraction, ForceCast_Exact_Replicate));
-					}
-				break;
 			case MP_ShadingModel:
 				return CompileShadingModelEncoding(Compiler, MaterialInterface->CompileProperty(&ProxyCompiler, MP_ShadingModel));
 			case MP_CustomOutput:
