Use engine's material baking module in glTF exporter

diff --git a/Engine/Plugins/Enterprise/GLTFExporter/GLTFExporter.uplugin b/Engine/Plugins/Enterprise/GLTFExporter/GLTFExporter.uplugin
--- a/Engine/Plugins/Enterprise/GLTFExporter/GLTFExporter.uplugin
+++ b/Engine/Plugins/Enterprise/GLTFExporter/GLTFExporter.uplugin
@@ -23,12 +23,6 @@
 			"Type": "Editor",
 			"LoadingPhase": "Default",
 			"PlatformAllowList": [ "Win64", "Mac", "Linux" ]
-		},
-		{
-			"Name": "GLTFMaterialBaking",
-			"Type": "Editor",
-			"LoadingPhase": "Default",
-			"PlatformAllowList": [ "Win64", "Mac", "Linux" ]
 		}
 	],
 	"Plugins": [
diff --git a/Engine/Plugins/Enterprise/GLTFExporter/Source/GLTFExporter/GLTFExporter.Build.cs b/Engine/Plugins/Enterprise/GLTFExporter/Source/GLTFExporter/GLTFExporter.Build.cs
--- a/Engine/Plugins/Enterprise/GLTFExporter/Source/GLTFExporter/GLTFExporter.Build.cs
+++ b/Engine/Plugins/Enterprise/GLTFExporter/Source/GLTFExporter/GLTFExporter.Build.cs
@@ -35,11 +35,11 @@ public class GLTFExporter : ModuleRules
 
 		if (Target.bBuildEditor)
 		{
-			// TODO: remove this when we no longer need to include GLTFMaterialBakingStructures.h in GLTFMeshData.h
+			// TODO: remove this when we no longer need to include MaterialBakingStructures.h in GLTFMeshData.h
 			PublicDependencyModuleNames.AddRange(
 				new string[]
 				{
-					"GLTFMaterialBaking",
+					"MaterialBaking",
 				}
 			);
 
diff --git a/Engine/Plugins/Enterprise/GLTFExporter/Source/GLTFExporter/Private/Converters/GLTFMaterialUtilities.cpp b/Engine/Plugins/Enterprise/GLTFExporter/Source/GLTFExporter/Private/Converters/GLTFMaterialUtilities.cpp
--- a/Engine/Plugins/Enterprise/GLTFExporter/Source/GLTFExporter/Private/Converters/GLTFMaterialUtilities.cpp
+++ b/Engine/Plugins/Enterprise/GLTFExporter/Source/GLTFExporter/Private/Converters/GLTFMaterialUtilities.cpp
@@ -8,8 +8,8 @@
 #include "Misc/DefaultValueHelper.h"
 #if WITH_EDITOR
 #include "GLTFMaterialAnalyzer.h"
-#include "IGLTFMaterialBakingModule.h"
-#include "GLTFMaterialBakingStructures.h"
+#include "IMaterialBakingModule.h"
+#include "MaterialBakingStructures.h"
 #include "Materials/MaterialExpressionTextureSample.h"
 #include "Materials/MaterialAttributeDefinitionMap.h"
 #endif
@@ -32,36 +32,36 @@ bool FGLTFMaterialUtilities::IsClearCoatBottomNormalEnabled()
 
 #if WITH_EDITOR
 
-bool FGLTFMaterialUtilities::IsNormalMap(const FGLTFMaterialPropertyEx& Property)
+bool FGLTFMaterialUtilities::IsNormalMap(const FMaterialPropertyEx& Property)
 {
-	return Property == MP_Normal || Property == FGLTFMaterialPropertyEx::ClearCoatBottomNormal;
+	return Property == MP_Normal || Property == FMaterialPropertyEx::ClearCoatBottomNormal;
 }
 
-bool FGLTFMaterialUtilities::IsSRGB(const FGLTFMaterialPropertyEx& Property)
+bool FGLTFMaterialUtilities::IsSRGB(const FMaterialPropertyEx& Property)
 {
-	return Property == MP_BaseColor || Property == MP_EmissiveColor || Property == MP_SubsurfaceColor || Property == FGLTFMaterialPropertyEx::TransmittanceColor;
+	return Property == MP_BaseColor || Property == MP_EmissiveColor || Property == MP_SubsurfaceColor || Property == FMaterialPropertyEx::TransmittanceColor;
 }
 
-FGuid FGLTFMaterialUtilities::GetAttributeID(const FGLTFMaterialPropertyEx& Property)
+FGuid FGLTFMaterialUtilities::GetAttributeID(const FMaterialPropertyEx& Property)
 {
 	return Property.IsCustomOutput()
 		? FMaterialAttributeDefinitionMap::GetCustomAttributeID(Property.CustomOutput.ToString())
 		: FMaterialAttributeDefinitionMap::GetID(Property.Type);
 }
 
-FGuid FGLTFMaterialUtilities::GetAttributeIDChecked(const FGLTFMaterialPropertyEx& Property)
+FGuid FGLTFMaterialUtilities::GetAttributeIDChecked(const FMaterialPropertyEx& Property)
 {
 	const FGuid AttributeID = GetAttributeID(Property);
 	check(AttributeID != FMaterialAttributeDefinitionMap::GetDefaultID());
 	return AttributeID;
 }
 
-FVector4f FGLTFMaterialUtilities::GetPropertyDefaultValue(const FGLTFMaterialPropertyEx& Property)
+FVector4f FGLTFMaterialUtilities::GetPropertyDefaultValue(const FMaterialPropertyEx& Property)
 {
 	return FMaterialAttributeDefinitionMap::GetDefaultValue(GetAttributeIDChecked(Property));
 }
 
-FVector4f FGLTFMaterialUtilities::GetPropertyMask(const FGLTFMaterialPropertyEx& Property)
+FVector4f FGLTFMaterialUtilities::GetPropertyMask(const FMaterialPropertyEx& Property)
 {
 	switch (FMaterialAttributeDefinitionMap::GetValueType(GetAttributeIDChecked(Property)))
 	{
@@ -76,7 +76,7 @@ FVector4f FGLTFMaterialUtilities::GetPropertyMask(const FGLTFMaterialPropertyEx&
 	}
 }
 
-const FExpressionInput* FGLTFMaterialUtilities::GetInputForProperty(const UMaterialInterface* Material, const FGLTFMaterialPropertyEx& Property)
+const FExpressionInput* FGLTFMaterialUtilities::GetInputForProperty(const UMaterialInterface* Material, const FMaterialPropertyEx& Property)
 {
 	if (Property.IsCustomOutput())
 	{
@@ -109,9 +109,9 @@ const UMaterialExpressionCustomOutput* FGLTFMaterialUtilities::GetCustomOutputBy
 	return nullptr;
 }
 
-FGLTFPropertyBakeOutput FGLTFMaterialUtilities::BakeMaterialProperty(const FIntPoint& OutputSize, const FGLTFMaterialPropertyEx& Property, const UMaterialInterface* Material, int32 TexCoord, const FGLTFMeshData* MeshData, const FGLTFIndexArray& MeshSectionIndices, bool bFillAlpha, bool bAdjustNormalmaps)
+FGLTFPropertyBakeOutput FGLTFMaterialUtilities::BakeMaterialProperty(const FIntPoint& OutputSize, const FMaterialPropertyEx& Property, const UMaterialInterface* Material, int32 TexCoord, const FGLTFMeshData* MeshData, const FGLTFIndexArray& MeshSectionIndices, bool bFillAlpha, bool bAdjustNormalmaps)
 {
-	FGLTFMeshRenderData MeshSet;
+	FMeshData MeshSet;
 	MeshSet.TextureCoordinateBox = { { 0.0f, 0.0f }, { 1.0f, 1.0f } };
 	MeshSet.TextureCoordinateIndex = TexCoord;
 	MeshSet.MaterialIndices = MeshSectionIndices; // NOTE: MaterialIndices is actually section indices
@@ -124,25 +124,25 @@ FGLTFPropertyBakeOutput FGLTFMaterialUtilities::BakeMaterialProperty(const FIntP
 		MeshSet.PrimitiveData = &MeshData->PrimitiveData;
 	}
 
-	FGLTFMaterialDataEx MatSet;
+	FMaterialDataEx MatSet;
 	MatSet.Material = const_cast<UMaterialInterface*>(Material);
 	MatSet.PropertySizes.Add(Property, OutputSize);
 	MatSet.bTangentSpaceNormal = true;
 
-	TArray<FGLTFMeshRenderData*> MeshSettings;
-	TArray<FGLTFMaterialDataEx*> MatSettings;
+	TArray<FMeshData*> MeshSettings;
+	TArray<FMaterialDataEx*> MatSettings;
 	MeshSettings.Add(&MeshSet);
 	MatSettings.Add(&MatSet);
 
-	TArray<FGLTFBakeOutputEx> BakeOutputs;
-	IGLTFMaterialBakingModule& Module = FModuleManager::Get().LoadModuleChecked<IGLTFMaterialBakingModule>("GLTFMaterialBaking");
+	TArray<FBakeOutputEx> BakeOutputs;
+	IMaterialBakingModule& Module = FModuleManager::Get().LoadModuleChecked<IMaterialBakingModule>("MaterialBaking");
 
 	Module.SetLinearBake(true);
 	Module.BakeMaterials(MatSettings, MeshSettings, BakeOutputs);
 	const bool bIsLinearBake = Module.IsLinearBake(Property);
 	Module.SetLinearBake(false);
 
-	FGLTFBakeOutputEx& BakeOutput = BakeOutputs[0];
+	FBakeOutputEx& BakeOutput = BakeOutputs[0];
 
 	TGLTFSharedArray<FColor> BakedPixels = MakeShared<TArray<FColor>>(MoveTemp(BakeOutput.PropertyData.FindChecked(Property)));
 	const FIntPoint BakedSize = BakeOutput.PropertySizes.FindChecked(Property);
@@ -228,7 +228,7 @@ bool FGLTFMaterialUtilities::TryGetTextureCoordinateIndex(const UMaterialExpress
 	return false;
 }
 
-void FGLTFMaterialUtilities::GetAllTextureCoordinateIndices(const UMaterialInterface* InMaterial, const FGLTFMaterialPropertyEx& InProperty, FGLTFIndexArray& OutTexCoords)
+void FGLTFMaterialUtilities::GetAllTextureCoordinateIndices(const UMaterialInterface* InMaterial, const FMaterialPropertyEx& InProperty, FGLTFIndexArray& OutTexCoords)
 {
 	FGLTFMaterialAnalysis Analysis;
 	AnalyzeMaterialProperty(InMaterial, InProperty, Analysis);
@@ -243,7 +243,7 @@ void FGLTFMaterialUtilities::GetAllTextureCoordinateIndices(const UMaterialInter
 	}
 }
 
-void FGLTFMaterialUtilities::AnalyzeMaterialProperty(const UMaterialInterface* InMaterial, const FGLTFMaterialPropertyEx& InProperty, FGLTFMaterialAnalysis& OutAnalysis)
+void FGLTFMaterialUtilities::AnalyzeMaterialProperty(const UMaterialInterface* InMaterial, const FMaterialPropertyEx& InProperty, FGLTFMaterialAnalysis& OutAnalysis)
 {
 	if (GetInputForProperty(InMaterial, InProperty) == nullptr)
 	{
@@ -315,7 +315,7 @@ bool FGLTFMaterialUtilities::NeedsMeshData(const UMaterialInterface* Material)
 	if (Material != nullptr && !FGLTFProxyMaterialUtilities::IsProxyMaterial(Material))
 	{
 		// TODO: only analyze properties that will be needed for this specific material
-		const TArray<FGLTFMaterialPropertyEx> Properties =
+		const TArray<FMaterialPropertyEx> Properties =
 		{
 			MP_BaseColor,
 			MP_EmissiveColor,
@@ -327,7 +327,7 @@ bool FGLTFMaterialUtilities::NeedsMeshData(const UMaterialInterface* Material)
 			MP_AmbientOcclusion,
 			MP_CustomData0,
 			MP_CustomData1,
-			FGLTFMaterialPropertyEx::ClearCoatBottomNormal,
+			FMaterialPropertyEx::ClearCoatBottomNormal,
 			// TODO: add TransmittanceColor when supported
 			// TODO: add Refraction when supported
 			// TODO: add Specular when supported
@@ -338,7 +338,7 @@ bool FGLTFMaterialUtilities::NeedsMeshData(const UMaterialInterface* Material)
 
 		// TODO: optimize baking by separating need for vertex data and primitive data
 
-		for (const FGLTFMaterialPropertyEx& Property: Properties)
+		for (const FMaterialPropertyEx& Property: Properties)
 		{
 			AnalyzeMaterialProperty(Material, Property, Analysis);
 			bNeedsMeshData |= Analysis.bRequiresVertexData;
diff --git a/Engine/Plugins/Enterprise/GLTFExporter/Source/GLTFExporter/Private/Converters/GLTFMaterialUtilities.h b/Engine/Plugins/Enterprise/GLTFExporter/Source/GLTFExporter/Private/Converters/GLTFMaterialUtilities.h
--- a/Engine/Plugins/Enterprise/GLTFExporter/Source/GLTFExporter/Private/Converters/GLTFMaterialUtilities.h
+++ b/Engine/Plugins/Enterprise/GLTFExporter/Source/GLTFExporter/Private/Converters/GLTFMaterialUtilities.h
@@ -7,7 +7,7 @@
 #include "MaterialExpressionIO.h"
 
 struct FGLTFMaterialAnalysis;
-struct FGLTFMaterialPropertyEx;
+struct FMaterialPropertyEx;
 class UMaterialExpressionTextureSample;
 class UMaterialInstanceConstant;
 class UMaterialInterface;
@@ -15,11 +15,11 @@ class UMaterialExpressionCustomOutput;
 
 struct FGLTFPropertyBakeOutput
 {
-	FGLTFPropertyBakeOutput(const FGLTFMaterialPropertyEx& Property, EPixelFormat PixelFormat, const TGLTFSharedArray<FColor>& Pixels, FIntPoint Size, float EmissiveScale, bool bIsSRGB)
+	FGLTFPropertyBakeOutput(const FMaterialPropertyEx& Property, EPixelFormat PixelFormat, const TGLTFSharedArray<FColor>& Pixels, FIntPoint Size, float EmissiveScale, bool bIsSRGB)
 		: Property(Property), PixelFormat(PixelFormat), Pixels(Pixels), Size(Size), EmissiveScale(EmissiveScale), bIsSRGB(bIsSRGB), bIsConstant(false)
 	{}
 
-	const FGLTFMaterialPropertyEx& Property;
+	const FMaterialPropertyEx& Property;
 	EPixelFormat PixelFormat;
 	TGLTFSharedArray<FColor> Pixels;
 	FIntPoint Size;
@@ -36,19 +36,19 @@ struct FGLTFMaterialUtilities
 	static bool IsClearCoatBottomNormalEnabled();
 
 #if WITH_EDITOR
-	static bool IsNormalMap(const FGLTFMaterialPropertyEx& Property);
-	static bool IsSRGB(const FGLTFMaterialPropertyEx& Property);
+	static bool IsNormalMap(const FMaterialPropertyEx& Property);
+	static bool IsSRGB(const FMaterialPropertyEx& Property);
 
-	static FGuid GetAttributeID(const FGLTFMaterialPropertyEx& Property);
-	static FGuid GetAttributeIDChecked(const FGLTFMaterialPropertyEx& Property);
+	static FGuid GetAttributeID(const FMaterialPropertyEx& Property);
+	static FGuid GetAttributeIDChecked(const FMaterialPropertyEx& Property);
 
-	static FVector4f GetPropertyDefaultValue(const FGLTFMaterialPropertyEx& Property);
-	static FVector4f GetPropertyMask(const FGLTFMaterialPropertyEx& Property);
+	static FVector4f GetPropertyDefaultValue(const FMaterialPropertyEx& Property);
+	static FVector4f GetPropertyMask(const FMaterialPropertyEx& Property);
 
-	static const FExpressionInput* GetInputForProperty(const UMaterialInterface* Material, const FGLTFMaterialPropertyEx& Property);
+	static const FExpressionInput* GetInputForProperty(const UMaterialInterface* Material, const FMaterialPropertyEx& Property);
 
 	template <class InputType>
-	static const FMaterialInput<InputType>* GetInputForProperty(const UMaterialInterface* Material, const FGLTFMaterialPropertyEx& Property)
+	static const FMaterialInput<InputType>* GetInputForProperty(const UMaterialInterface* Material, const FMaterialPropertyEx& Property)
 	{
 		const FExpressionInput* ExpressionInput = GetInputForProperty(Material, Property);
 		return static_cast<const FMaterialInput<InputType>*>(ExpressionInput);
@@ -56,7 +56,7 @@ struct FGLTFMaterialUtilities
 
 	static const UMaterialExpressionCustomOutput* GetCustomOutputByName(const UMaterialInterface* Material, const FString& FunctionName);
 
-	static FGLTFPropertyBakeOutput BakeMaterialProperty(const FIntPoint& OutputSize, const FGLTFMaterialPropertyEx& Property, const UMaterialInterface* Material, int32 TexCoord, const FGLTFMeshData* MeshData = nullptr, const FGLTFIndexArray& MeshSectionIndices = {}, bool bFillAlpha = true, bool bAdjustNormalmaps = true);
+	static FGLTFPropertyBakeOutput BakeMaterialProperty(const FIntPoint& OutputSize, const FMaterialPropertyEx& Property, const UMaterialInterface* Material, int32 TexCoord, const FGLTFMeshData* MeshData = nullptr, const FGLTFIndexArray& MeshSectionIndices = {}, bool bFillAlpha = true, bool bAdjustNormalmaps = true);
 
 	static FGLTFJsonTexture* AddTexture(FGLTFConvertBuilder& Builder, TGLTFSharedArray<FColor>& Pixels, const FIntPoint& TextureSize, bool bIgnoreAlpha, bool bIsNormalMap, const FString& TextureName, TextureAddress TextureAddress, TextureFilter TextureFilter);
 
@@ -64,9 +64,9 @@ struct FGLTFMaterialUtilities
 	static uint32 GetMaskComponentCount(const FExpressionInput& ExpressionInput);
 
 	static bool TryGetTextureCoordinateIndex(const UMaterialExpressionTextureSample* TextureSampler, int32& TexCoord, FGLTFJsonTextureTransform& Transform);
-	static void GetAllTextureCoordinateIndices(const UMaterialInterface* InMaterial, const FGLTFMaterialPropertyEx& InProperty, FGLTFIndexArray& OutTexCoords);
+	static void GetAllTextureCoordinateIndices(const UMaterialInterface* InMaterial, const FMaterialPropertyEx& InProperty, FGLTFIndexArray& OutTexCoords);
 
-	static void AnalyzeMaterialProperty(const UMaterialInterface* Material, const FGLTFMaterialPropertyEx& InProperty, FGLTFMaterialAnalysis& OutAnalysis);
+	static void AnalyzeMaterialProperty(const UMaterialInterface* Material, const FMaterialPropertyEx& InProperty, FGLTFMaterialAnalysis& OutAnalysis);
 
 	static FMaterialShadingModelField EvaluateShadingModelExpression(const UMaterialInterface* Material);
 #endif
diff --git a/Engine/Plugins/Enterprise/GLTFExporter/Source/GLTFExporter/Private/Converters/GLTFUVOverlapChecker.cpp b/Engine/Plugins/Enterprise/GLTFExporter/Source/GLTFExporter/Private/Converters/GLTFUVOverlapChecker.cpp
--- a/Engine/Plugins/Enterprise/GLTFExporter/Source/GLTFExporter/Private/Converters/GLTFUVOverlapChecker.cpp
+++ b/Engine/Plugins/Enterprise/GLTFExporter/Source/GLTFExporter/Private/Converters/GLTFUVOverlapChecker.cpp
@@ -6,8 +6,8 @@
 #include "StaticMeshAttributes.h"
 #include "MeshDescription.h"
 #include "Modules/ModuleManager.h"
-#include "IGLTFMaterialBakingModule.h"
-#include "GLTFMaterialBakingStructures.h"
+#include "IMaterialBakingModule.h"
+#include "MaterialBakingStructures.h"
 #else
 #include "UObject/UObjectGlobals.h"
 #endif
@@ -19,32 +19,32 @@ float FGLTFUVOverlapChecker::Convert(const FMeshDescription* Description, FGLTFI
 	{
 		// TODO: investigate if the fixed size is high enough to properly calculate overlap
 		const FIntPoint TextureSize(512, 512);
-		const FGLTFMaterialPropertyEx Property = MP_Opacity;
+		const FMaterialPropertyEx Property = MP_Opacity;
 
-		FGLTFMeshRenderData MeshSet;
+		FMeshData MeshSet;
 		MeshSet.TextureCoordinateBox = { { 0.0f, 0.0f }, { 1.0f, 1.0f } };
 		MeshSet.TextureCoordinateIndex = TexCoord;
 		MeshSet.MeshDescription = const_cast<FMeshDescription*>(Description);
 		MeshSet.MaterialIndices = SectionIndices; // NOTE: MaterialIndices is actually section indices
 
-		FGLTFMaterialDataEx MatSet;
+		FMaterialDataEx MatSet;
 		MatSet.Material = GetMaterial();
 		MatSet.PropertySizes.Add(Property, TextureSize);
 		MatSet.BlendMode = MatSet.Material->GetBlendMode();
 		MatSet.BackgroundColor = FColor::Black;
 		MatSet.bPerformBorderSmear = false;
 
-		TArray<FGLTFMeshRenderData*> MeshSettings;
-		TArray<FGLTFMaterialDataEx*> MatSettings;
+		TArray<FMeshData*> MeshSettings;
+		TArray<FMaterialDataEx*> MatSettings;
 		MeshSettings.Add(&MeshSet);
 		MatSettings.Add(&MatSet);
 
-		TArray<FGLTFBakeOutputEx> BakeOutputs;
-		IGLTFMaterialBakingModule& Module = FModuleManager::Get().LoadModuleChecked<IGLTFMaterialBakingModule>("GLTFMaterialBaking");
+		TArray<FBakeOutputEx> BakeOutputs;
+		IMaterialBakingModule& Module = FModuleManager::Get().LoadModuleChecked<IMaterialBakingModule>("MaterialBaking");
 
 		Module.BakeMaterials(MatSettings, MeshSettings, BakeOutputs);
 
-		const FGLTFBakeOutputEx& BakeOutput = BakeOutputs[0];
+		const FBakeOutputEx& BakeOutput = BakeOutputs[0];
 		const TArray<FColor>& BakedPixels = BakeOutput.PropertyData.FindChecked(Property);
 
 		if (BakedPixels.Num() <= 0)
diff --git a/Engine/Plugins/Enterprise/GLTFExporter/Source/GLTFExporter/Private/Tasks/GLTFDelayedMaterialTasks.cpp b/Engine/Plugins/Enterprise/GLTFExporter/Source/GLTFExporter/Private/Tasks/GLTFDelayedMaterialTasks.cpp
--- a/Engine/Plugins/Enterprise/GLTFExporter/Source/GLTFExporter/Private/Tasks/GLTFDelayedMaterialTasks.cpp
+++ b/Engine/Plugins/Enterprise/GLTFExporter/Source/GLTFExporter/Private/Tasks/GLTFDelayedMaterialTasks.cpp
@@ -72,8 +72,8 @@ void FGLTFDelayedMaterialTask::Process()
 #if WITH_EDITOR
 	if (JsonMaterial->ShadingModel != EGLTFJsonShadingModel::None)
 	{
-		const FGLTFMaterialPropertyEx BaseColorProperty = JsonMaterial->ShadingModel == EGLTFJsonShadingModel::Unlit ? MP_EmissiveColor : MP_BaseColor;
-		const FGLTFMaterialPropertyEx OpacityProperty = JsonMaterial->AlphaMode == EGLTFJsonAlphaMode::Mask ? MP_OpacityMask : MP_Opacity;
+		const FMaterialPropertyEx BaseColorProperty = JsonMaterial->ShadingModel == EGLTFJsonShadingModel::Unlit ? MP_EmissiveColor : MP_BaseColor;
+		const FMaterialPropertyEx OpacityProperty = JsonMaterial->AlphaMode == EGLTFJsonAlphaMode::Mask ? MP_OpacityMask : MP_Opacity;
 
 		// TODO: check if a property is active before trying to get it (i.e. Material->IsPropertyActive)
 
@@ -102,21 +102,21 @@ void FGLTFDelayedMaterialTask::Process()
 
 		if (JsonMaterial->ShadingModel == EGLTFJsonShadingModel::Default || JsonMaterial->ShadingModel == EGLTFJsonShadingModel::ClearCoat)
 		{
-			const FGLTFMaterialPropertyEx MetallicProperty = MP_Metallic;
-			const FGLTFMaterialPropertyEx RoughnessProperty = MP_Roughness;
+			const FMaterialPropertyEx MetallicProperty = MP_Metallic;
+			const FMaterialPropertyEx RoughnessProperty = MP_Roughness;
 
 			if (!TryGetMetallicAndRoughness(JsonMaterial->PBRMetallicRoughness, MetallicProperty, RoughnessProperty))
 			{
 				Builder.LogWarning(FString::Printf(TEXT("Failed to export %s and %s for material %s"), *MetallicProperty.ToString(), *RoughnessProperty.ToString(), *Material->GetName()));
 			}
 
-			const FGLTFMaterialPropertyEx EmissiveProperty = MP_EmissiveColor;
+			const FMaterialPropertyEx EmissiveProperty = MP_EmissiveColor;
 			if (!TryGetEmissive(*JsonMaterial, EmissiveProperty))
 			{
 				Builder.LogWarning(FString::Printf(TEXT("Failed to export %s for material %s"), *EmissiveProperty.ToString(), *Material->GetName()));
 			}
 
-			const FGLTFMaterialPropertyEx NormalProperty = JsonMaterial->ShadingModel == EGLTFJsonShadingModel::ClearCoat ? FGLTFMaterialPropertyEx::ClearCoatBottomNormal : FGLTFMaterialPropertyEx(MP_Normal);
+			const FMaterialPropertyEx NormalProperty = JsonMaterial->ShadingModel == EGLTFJsonShadingModel::ClearCoat ? FMaterialPropertyEx::ClearCoatBottomNormal : FMaterialPropertyEx(MP_Normal);
 			if (IsPropertyNonDefault(NormalProperty))
 			{
 				if (!TryGetSourceTexture(JsonMaterial->NormalTexture, NormalProperty, DefaultColorInputMasks))
@@ -128,7 +128,7 @@ void FGLTFDelayedMaterialTask::Process()
 				}
 			}
 
-			const FGLTFMaterialPropertyEx AmbientOcclusionProperty = MP_AmbientOcclusion;
+			const FMaterialPropertyEx AmbientOcclusionProperty = MP_AmbientOcclusion;
 			if (IsPropertyNonDefault(AmbientOcclusionProperty))
 			{
 				if (!TryGetSourceTexture(JsonMaterial->OcclusionTexture, AmbientOcclusionProperty, OcclusionInputMasks))
@@ -142,15 +142,15 @@ void FGLTFDelayedMaterialTask::Process()
 
 			if (JsonMaterial->ShadingModel == EGLTFJsonShadingModel::ClearCoat)
 			{
-				const FGLTFMaterialPropertyEx ClearCoatProperty = MP_CustomData0;
-				const FGLTFMaterialPropertyEx ClearCoatRoughnessProperty = MP_CustomData1;
+				const FMaterialPropertyEx ClearCoatProperty = MP_CustomData0;
+				const FMaterialPropertyEx ClearCoatRoughnessProperty = MP_CustomData1;
 
 				if (!TryGetClearCoatRoughness(JsonMaterial->ClearCoat, ClearCoatProperty, ClearCoatRoughnessProperty))
 				{
 					Builder.LogWarning(FString::Printf(TEXT("Failed to export %s and %s for material %s"), *ClearCoatProperty.ToString(), *ClearCoatRoughnessProperty.ToString(), *Material->GetName()));
 				}
 
-				const FGLTFMaterialPropertyEx ClearCoatNormalProperty = MP_Normal;
+				const FMaterialPropertyEx ClearCoatNormalProperty = MP_Normal;
 				if (IsPropertyNonDefault(ClearCoatNormalProperty))
 				{
 					if (!TryGetSourceTexture(JsonMaterial->ClearCoat.ClearCoatNormalTexture, ClearCoatNormalProperty, DefaultColorInputMasks))
@@ -468,7 +468,7 @@ void FGLTFDelayedMaterialTask::ConvertAlphaMode(EGLTFJsonAlphaMode& OutAlphaMode
 
 #if WITH_EDITOR
 
-bool FGLTFDelayedMaterialTask::TryGetBaseColorAndOpacity(FGLTFJsonPBRMetallicRoughness& OutPBRParams, const FGLTFMaterialPropertyEx& BaseColorProperty, const FGLTFMaterialPropertyEx& OpacityProperty)
+bool FGLTFDelayedMaterialTask::TryGetBaseColorAndOpacity(FGLTFJsonPBRMetallicRoughness& OutPBRParams, const FMaterialPropertyEx& BaseColorProperty, const FMaterialPropertyEx& OpacityProperty)
 {
 	const bool bIsBaseColorConstant = TryGetConstantColor(OutPBRParams.BaseColorFactor, BaseColorProperty);
 	const bool bIsOpacityConstant = TryGetConstantScalar(OutPBRParams.BaseColorFactor.A, OpacityProperty);
@@ -585,7 +585,7 @@ bool FGLTFDelayedMaterialTask::TryGetBaseColorAndOpacity(FGLTFJsonPBRMetallicRou
 	return true;
 }
 
-bool FGLTFDelayedMaterialTask::TryGetMetallicAndRoughness(FGLTFJsonPBRMetallicRoughness& OutPBRParams, const FGLTFMaterialPropertyEx& MetallicProperty, const FGLTFMaterialPropertyEx& RoughnessProperty)
+bool FGLTFDelayedMaterialTask::TryGetMetallicAndRoughness(FGLTFJsonPBRMetallicRoughness& OutPBRParams, const FMaterialPropertyEx& MetallicProperty, const FMaterialPropertyEx& RoughnessProperty)
 {
 	const bool bIsMetallicConstant = TryGetConstantScalar(OutPBRParams.MetallicFactor, MetallicProperty);
 	const bool bIsRoughnessConstant = TryGetConstantScalar(OutPBRParams.RoughnessFactor, RoughnessProperty);
@@ -698,7 +698,7 @@ bool FGLTFDelayedMaterialTask::TryGetMetallicAndRoughness(FGLTFJsonPBRMetallicRo
 	return true;
 }
 
-bool FGLTFDelayedMaterialTask::TryGetClearCoatRoughness(FGLTFJsonClearCoatExtension& OutExtParams, const FGLTFMaterialPropertyEx& IntensityProperty, const FGLTFMaterialPropertyEx& RoughnessProperty)
+bool FGLTFDelayedMaterialTask::TryGetClearCoatRoughness(FGLTFJsonClearCoatExtension& OutExtParams, const FMaterialPropertyEx& IntensityProperty, const FMaterialPropertyEx& RoughnessProperty)
 {
 	const bool bIsIntensityConstant = TryGetConstantScalar(OutExtParams.ClearCoatFactor, IntensityProperty);
 	const bool bIsRoughnessConstant = TryGetConstantScalar(OutExtParams.ClearCoatRoughnessFactor, RoughnessProperty);
@@ -815,7 +815,7 @@ bool FGLTFDelayedMaterialTask::TryGetClearCoatRoughness(FGLTFJsonClearCoatExtens
 	return true;
 }
 
-bool FGLTFDelayedMaterialTask::TryGetEmissive(FGLTFJsonMaterial& OutMaterial, const FGLTFMaterialPropertyEx& EmissiveProperty)
+bool FGLTFDelayedMaterialTask::TryGetEmissive(FGLTFJsonMaterial& OutMaterial, const FMaterialPropertyEx& EmissiveProperty)
 {
 	// TODO: right now we allow EmissiveFactor to be > 1.0 to support very bright emission, although it's not valid according to the glTF standard.
 	// We may want to change this behaviour and store factors above 1.0 using a custom extension instead.
@@ -867,7 +867,7 @@ bool FGLTFDelayedMaterialTask::TryGetEmissive(FGLTFJsonMaterial& OutMaterial, co
 	return true;
 }
 
-bool FGLTFDelayedMaterialTask::IsPropertyNonDefault(const FGLTFMaterialPropertyEx& Property) const
+bool FGLTFDelayedMaterialTask::IsPropertyNonDefault(const FMaterialPropertyEx& Property) const
 {
 	// Custom outputs are by definition standalone and never part of material attributes
 	if (!Property.IsCustomOutput())
@@ -892,7 +892,7 @@ bool FGLTFDelayedMaterialTask::IsPropertyNonDefault(const FGLTFMaterialPropertyE
 		return false;
 	}
 
-	if (Property == FGLTFMaterialPropertyEx::ClearCoatBottomNormal && !FGLTFMaterialUtilities::IsClearCoatBottomNormalEnabled())
+	if (Property == FMaterialPropertyEx::ClearCoatBottomNormal && !FGLTFMaterialUtilities::IsClearCoatBottomNormalEnabled())
 	{
 		Builder.LogWarning(FString::Printf(
 			TEXT("Material %s won't be exported with clear coat bottom normal because ClearCoatEnableSecondNormal in project rendering settings is disabled"),
@@ -903,7 +903,7 @@ bool FGLTFDelayedMaterialTask::IsPropertyNonDefault(const FGLTFMaterialPropertyE
 	return true;
 }
 
-bool FGLTFDelayedMaterialTask::TryGetConstantColor(FGLTFJsonColor3& OutValue, const FGLTFMaterialPropertyEx& Property) const
+bool FGLTFDelayedMaterialTask::TryGetConstantColor(FGLTFJsonColor3& OutValue, const FMaterialPropertyEx& Property) const
 {
 	FLinearColor Value;
 	if (TryGetConstantColor(Value, Property))
@@ -915,7 +915,7 @@ bool FGLTFDelayedMaterialTask::TryGetConstantColor(FGLTFJsonColor3& OutValue, co
 	return false;
 }
 
-bool FGLTFDelayedMaterialTask::TryGetConstantColor(FGLTFJsonColor4& OutValue, const FGLTFMaterialPropertyEx& Property) const
+bool FGLTFDelayedMaterialTask::TryGetConstantColor(FGLTFJsonColor4& OutValue, const FMaterialPropertyEx& Property) const
 {
 	FLinearColor Value;
 	if (TryGetConstantColor(Value, Property))
@@ -927,7 +927,7 @@ bool FGLTFDelayedMaterialTask::TryGetConstantColor(FGLTFJsonColor4& OutValue, co
 	return false;
 }
 
-bool FGLTFDelayedMaterialTask::TryGetConstantColor(FLinearColor& OutValue, const FGLTFMaterialPropertyEx& Property) const
+bool FGLTFDelayedMaterialTask::TryGetConstantColor(FLinearColor& OutValue, const FMaterialPropertyEx& Property) const
 {
 	const bool bUseMaterialAttributes = Material->GetMaterial()->bUseMaterialAttributes;
 	if (bUseMaterialAttributes)
@@ -1034,7 +1034,7 @@ bool FGLTFDelayedMaterialTask::TryGetConstantColor(FLinearColor& OutValue, const
 	return false;
 }
 
-bool FGLTFDelayedMaterialTask::TryGetConstantScalar(float& OutValue, const FGLTFMaterialPropertyEx& Property) const
+bool FGLTFDelayedMaterialTask::TryGetConstantScalar(float& OutValue, const FMaterialPropertyEx& Property) const
 {
 	const bool bUseMaterialAttributes = Material->GetMaterial()->bUseMaterialAttributes;
 	if (bUseMaterialAttributes)
@@ -1135,7 +1135,7 @@ bool FGLTFDelayedMaterialTask::TryGetConstantScalar(float& OutValue, const FGLTF
 	return false;
 }
 
-bool FGLTFDelayedMaterialTask::TryGetSourceTexture(FGLTFJsonTextureInfo& OutTexInfo, const FGLTFMaterialPropertyEx& Property, const TArray<FLinearColor>& AllowedMasks) const
+bool FGLTFDelayedMaterialTask::TryGetSourceTexture(FGLTFJsonTextureInfo& OutTexInfo, const FMaterialPropertyEx& Property, const TArray<FLinearColor>& AllowedMasks) const
 {
 	const UTexture2D* Texture;
 	int32 TexCoord;
@@ -1152,7 +1152,7 @@ bool FGLTFDelayedMaterialTask::TryGetSourceTexture(FGLTFJsonTextureInfo& OutTexI
 	return false;
 }
 
-bool FGLTFDelayedMaterialTask::TryGetSourceTexture(const UTexture2D*& OutTexture, int32& OutTexCoord, FGLTFJsonTextureTransform& OutTransform, const FGLTFMaterialPropertyEx& Property, const TArray<FLinearColor>& AllowedMasks) const
+bool FGLTFDelayedMaterialTask::TryGetSourceTexture(const UTexture2D*& OutTexture, int32& OutTexCoord, FGLTFJsonTextureTransform& OutTransform, const FMaterialPropertyEx& Property, const TArray<FLinearColor>& AllowedMasks) const
 {
 	if (Builder.ExportOptions->TextureImageFormat == EGLTFTextureImageFormat::None)
 	{
@@ -1275,7 +1275,7 @@ bool FGLTFDelayedMaterialTask::TryGetSourceTexture(const UTexture2D*& OutTexture
 	return false;
 }
 
-bool FGLTFDelayedMaterialTask::TryGetBakedMaterialProperty(FGLTFJsonTextureInfo& OutTexInfo, FGLTFJsonColor3& OutConstant, const FGLTFMaterialPropertyEx& Property, const FString& PropertyName)
+bool FGLTFDelayedMaterialTask::TryGetBakedMaterialProperty(FGLTFJsonTextureInfo& OutTexInfo, FGLTFJsonColor3& OutConstant, const FMaterialPropertyEx& Property, const FString& PropertyName)
 {
 	if (Builder.ExportOptions->BakeMaterialInputs == EGLTFMaterialBakeMode::Disabled)
 	{
@@ -1309,7 +1309,7 @@ bool FGLTFDelayedMaterialTask::TryGetBakedMaterialProperty(FGLTFJsonTextureInfo&
 	return false;
 }
 
-bool FGLTFDelayedMaterialTask::TryGetBakedMaterialProperty(FGLTFJsonTextureInfo& OutTexInfo, FGLTFJsonColor4& OutConstant, const FGLTFMaterialPropertyEx& Property, const FString& PropertyName)
+bool FGLTFDelayedMaterialTask::TryGetBakedMaterialProperty(FGLTFJsonTextureInfo& OutTexInfo, FGLTFJsonColor4& OutConstant, const FMaterialPropertyEx& Property, const FString& PropertyName)
 {
 	if (Builder.ExportOptions->BakeMaterialInputs == EGLTFMaterialBakeMode::Disabled)
 	{
@@ -1343,7 +1343,7 @@ bool FGLTFDelayedMaterialTask::TryGetBakedMaterialProperty(FGLTFJsonTextureInfo&
 	return false;
 }
 
-inline bool FGLTFDelayedMaterialTask::TryGetBakedMaterialProperty(FGLTFJsonTextureInfo& OutTexInfo, float& OutConstant, const FGLTFMaterialPropertyEx& Property, const FString& PropertyName)
+inline bool FGLTFDelayedMaterialTask::TryGetBakedMaterialProperty(FGLTFJsonTextureInfo& OutTexInfo, float& OutConstant, const FMaterialPropertyEx& Property, const FString& PropertyName)
 {
 	if (Builder.ExportOptions->BakeMaterialInputs == EGLTFMaterialBakeMode::Disabled)
 	{
@@ -1377,7 +1377,7 @@ inline bool FGLTFDelayedMaterialTask::TryGetBakedMaterialProperty(FGLTFJsonTextu
 	return false;
 }
 
-bool FGLTFDelayedMaterialTask::TryGetBakedMaterialProperty(FGLTFJsonTextureInfo& OutTexInfo, const FGLTFMaterialPropertyEx& Property, const FString& PropertyName)
+bool FGLTFDelayedMaterialTask::TryGetBakedMaterialProperty(FGLTFJsonTextureInfo& OutTexInfo, const FMaterialPropertyEx& Property, const FString& PropertyName)
 {
 	if (Builder.ExportOptions->BakeMaterialInputs == EGLTFMaterialBakeMode::Disabled)
 	{
@@ -1437,13 +1437,13 @@ bool FGLTFDelayedMaterialTask::TryGetBakedMaterialProperty(FGLTFJsonTextureInfo&
 	return true;
 }
 
-FGLTFPropertyBakeOutput FGLTFDelayedMaterialTask::BakeMaterialProperty(const FGLTFMaterialPropertyEx& Property, int32& OutTexCoord)
+FGLTFPropertyBakeOutput FGLTFDelayedMaterialTask::BakeMaterialProperty(const FMaterialPropertyEx& Property, int32& OutTexCoord)
 {
 	const FIntPoint TextureSize = Builder.GetBakeSizeForMaterialProperty(Material, GetPropertyGroup(Property));
 	return BakeMaterialProperty(Property, OutTexCoord, TextureSize, true);
 }
 
-FGLTFPropertyBakeOutput FGLTFDelayedMaterialTask::BakeMaterialProperty(const FGLTFMaterialPropertyEx& Property, int32& OutTexCoord, const FIntPoint& TextureSize, bool bFillAlpha)
+FGLTFPropertyBakeOutput FGLTFDelayedMaterialTask::BakeMaterialProperty(const FMaterialPropertyEx& Property, int32& OutTexCoord, const FIntPoint& TextureSize, bool bFillAlpha)
 {
 	if (MeshData == nullptr)
 	{
@@ -1507,7 +1507,7 @@ bool FGLTFDelayedMaterialTask::StoreBakedPropertyTexture(FGLTFJsonTextureInfo& O
 	return true;
 }
 
-EGLTFMaterialPropertyGroup FGLTFDelayedMaterialTask::GetPropertyGroup(const FGLTFMaterialPropertyEx& Property)
+EGLTFMaterialPropertyGroup FGLTFDelayedMaterialTask::GetPropertyGroup(const FMaterialPropertyEx& Property)
 {
 	switch (Property.Type)
 	{
@@ -1528,7 +1528,7 @@ EGLTFMaterialPropertyGroup FGLTFDelayedMaterialTask::GetPropertyGroup(const FGLT
 		case MP_CustomData1:
 			return EGLTFMaterialPropertyGroup::ClearCoatRoughness;
 		case MP_CustomOutput:
-			if (Property == FGLTFMaterialPropertyEx::ClearCoatBottomNormal)
+			if (Property == FMaterialPropertyEx::ClearCoatBottomNormal)
 			{
 				return EGLTFMaterialPropertyGroup::ClearCoatBottomNormal;
 			}
diff --git a/Engine/Plugins/Enterprise/GLTFExporter/Source/GLTFExporter/Private/Tasks/GLTFDelayedMaterialTasks.h b/Engine/Plugins/Enterprise/GLTFExporter/Source/GLTFExporter/Private/Tasks/GLTFDelayedMaterialTasks.h
--- a/Engine/Plugins/Enterprise/GLTFExporter/Source/GLTFExporter/Private/Tasks/GLTFDelayedMaterialTasks.h
+++ b/Engine/Plugins/Enterprise/GLTFExporter/Source/GLTFExporter/Private/Tasks/GLTFDelayedMaterialTasks.h
@@ -9,7 +9,7 @@
 #include "Materials/GLTFProxyMaterialParameterInfo.h"
 #include "Materials/Material.h"
 #if WITH_EDITOR
-#include "GLTFMaterialPropertyEx.h"
+#include "MaterialPropertyEx.h"
 #endif
 
 struct FGLTFPropertyBakeOutput;
@@ -66,33 +66,33 @@ private:
 	void ConvertAlphaMode(EGLTFJsonAlphaMode& OutAlphaMode) const;
 
 #if WITH_EDITOR
-	TSet<FGLTFMaterialPropertyEx> MeshDataBakedProperties;
+	TSet<FMaterialPropertyEx> MeshDataBakedProperties;
 
-	bool TryGetBaseColorAndOpacity(FGLTFJsonPBRMetallicRoughness& OutPBRParams, const FGLTFMaterialPropertyEx& BaseColorProperty, const FGLTFMaterialPropertyEx& OpacityProperty);
-	bool TryGetMetallicAndRoughness(FGLTFJsonPBRMetallicRoughness& OutPBRParams, const FGLTFMaterialPropertyEx& MetallicProperty, const FGLTFMaterialPropertyEx& RoughnessProperty);
-	bool TryGetClearCoatRoughness(FGLTFJsonClearCoatExtension& OutExtParams, const FGLTFMaterialPropertyEx& IntensityProperty, const FGLTFMaterialPropertyEx& RoughnessProperty);
-	bool TryGetEmissive(FGLTFJsonMaterial& OutMaterial, const FGLTFMaterialPropertyEx& EmissiveProperty);
+	bool TryGetBaseColorAndOpacity(FGLTFJsonPBRMetallicRoughness& OutPBRParams, const FMaterialPropertyEx& BaseColorProperty, const FMaterialPropertyEx& OpacityProperty);
+	bool TryGetMetallicAndRoughness(FGLTFJsonPBRMetallicRoughness& OutPBRParams, const FMaterialPropertyEx& MetallicProperty, const FMaterialPropertyEx& RoughnessProperty);
+	bool TryGetClearCoatRoughness(FGLTFJsonClearCoatExtension& OutExtParams, const FMaterialPropertyEx& IntensityProperty, const FMaterialPropertyEx& RoughnessProperty);
+	bool TryGetEmissive(FGLTFJsonMaterial& OutMaterial, const FMaterialPropertyEx& EmissiveProperty);
 
-	bool IsPropertyNonDefault(const FGLTFMaterialPropertyEx& Property) const;
-	bool TryGetConstantColor(FGLTFJsonColor3& OutValue, const FGLTFMaterialPropertyEx& Property) const;
-	bool TryGetConstantColor(FGLTFJsonColor4& OutValue, const FGLTFMaterialPropertyEx& Property) const;
-	bool TryGetConstantColor(FLinearColor& OutValue, const FGLTFMaterialPropertyEx& Property) const;
-	bool TryGetConstantScalar(float& OutValue, const FGLTFMaterialPropertyEx& Property) const;
+	bool IsPropertyNonDefault(const FMaterialPropertyEx& Property) const;
+	bool TryGetConstantColor(FGLTFJsonColor3& OutValue, const FMaterialPropertyEx& Property) const;
+	bool TryGetConstantColor(FGLTFJsonColor4& OutValue, const FMaterialPropertyEx& Property) const;
+	bool TryGetConstantColor(FLinearColor& OutValue, const FMaterialPropertyEx& Property) const;
+	bool TryGetConstantScalar(float& OutValue, const FMaterialPropertyEx& Property) const;
 
-	bool TryGetSourceTexture(FGLTFJsonTextureInfo& OutTexInfo, const FGLTFMaterialPropertyEx& Property, const TArray<FLinearColor>& AllowedMasks = {}) const;
-	bool TryGetSourceTexture(const UTexture2D*& OutTexture, int32& OutTexCoord, FGLTFJsonTextureTransform& OutTransform, const FGLTFMaterialPropertyEx& Property, const TArray<FLinearColor>& AllowedMasks = {}) const;
+	bool TryGetSourceTexture(FGLTFJsonTextureInfo& OutTexInfo, const FMaterialPropertyEx& Property, const TArray<FLinearColor>& AllowedMasks = {}) const;
+	bool TryGetSourceTexture(const UTexture2D*& OutTexture, int32& OutTexCoord, FGLTFJsonTextureTransform& OutTransform, const FMaterialPropertyEx& Property, const TArray<FLinearColor>& AllowedMasks = {}) const;
 
-	bool TryGetBakedMaterialProperty(FGLTFJsonTextureInfo& OutTexInfo, FGLTFJsonColor3& OutConstant, const FGLTFMaterialPropertyEx& Property, const FString& PropertyName);
-	bool TryGetBakedMaterialProperty(FGLTFJsonTextureInfo& OutTexInfo, FGLTFJsonColor4& OutConstant, const FGLTFMaterialPropertyEx& Property, const FString& PropertyName);
-	bool TryGetBakedMaterialProperty(FGLTFJsonTextureInfo& OutTexInfo, float& OutConstant, const FGLTFMaterialPropertyEx& Property, const FString& PropertyName);
-	bool TryGetBakedMaterialProperty(FGLTFJsonTextureInfo& OutTexInfo, const FGLTFMaterialPropertyEx& Property, const FString& PropertyName);
+	bool TryGetBakedMaterialProperty(FGLTFJsonTextureInfo& OutTexInfo, FGLTFJsonColor3& OutConstant, const FMaterialPropertyEx& Property, const FString& PropertyName);
+	bool TryGetBakedMaterialProperty(FGLTFJsonTextureInfo& OutTexInfo, FGLTFJsonColor4& OutConstant, const FMaterialPropertyEx& Property, const FString& PropertyName);
+	bool TryGetBakedMaterialProperty(FGLTFJsonTextureInfo& OutTexInfo, float& OutConstant, const FMaterialPropertyEx& Property, const FString& PropertyName);
+	bool TryGetBakedMaterialProperty(FGLTFJsonTextureInfo& OutTexInfo, const FMaterialPropertyEx& Property, const FString& PropertyName);
 
-	FGLTFPropertyBakeOutput BakeMaterialProperty(const FGLTFMaterialPropertyEx& Property, int32& OutTexCoord);
-	FGLTFPropertyBakeOutput BakeMaterialProperty(const FGLTFMaterialPropertyEx& Property, int32& OutTexCoord, const FIntPoint& TextureSize, bool bFillAlpha);
+	FGLTFPropertyBakeOutput BakeMaterialProperty(const FMaterialPropertyEx& Property, int32& OutTexCoord);
+	FGLTFPropertyBakeOutput BakeMaterialProperty(const FMaterialPropertyEx& Property, int32& OutTexCoord, const FIntPoint& TextureSize, bool bFillAlpha);
 
 	bool StoreBakedPropertyTexture(FGLTFJsonTextureInfo& OutTexInfo, FGLTFPropertyBakeOutput& PropertyBakeOutput, const FString& PropertyName) const;
 
-	static EGLTFMaterialPropertyGroup GetPropertyGroup(const FGLTFMaterialPropertyEx& Property);
+	static EGLTFMaterialPropertyGroup GetPropertyGroup(const FMaterialPropertyEx& Property);
 
 	template <typename CallbackType>
 	static void CombinePixels(const TArray<FColor>& FirstPixels, const TArray<FColor>& SecondPixels, TArray<FColor>& OutPixels, CallbackType Callback);
diff --git a/Engine/Plugins/Enterprise/GLTFExporter/Source/GLTFExporter/Public/Converters/GLTFMeshData.h b/Engine/Plugins/Enterprise/GLTFExporter/Source/GLTFExporter/Public/Converters/GLTFMeshData.h
--- a/Engine/Plugins/Enterprise/GLTFExporter/Source/GLTFExporter/Public/Converters/GLTFMeshData.h
+++ b/Engine/Plugins/Enterprise/GLTFExporter/Source/GLTFExporter/Public/Converters/GLTFMeshData.h
@@ -7,8 +7,8 @@
 #include "SceneTypes.h"
 
 #if WITH_EDITOR
-// TODO: remove dependency to GLTFMaterialBaking in public header
-#include "GLTFMaterialBakingStructures.h"
+// TODO: remove dependency to MaterialBaking in public header
+#include "MaterialBakingStructures.h"
 #endif
 
 class FLightmapResourceCluster;
@@ -31,7 +31,7 @@ struct GLTFEXPORTER_API FGLTFMeshData
 	int32 LODIndex;
 
 #if WITH_EDITOR
-	FGLTFPrimitiveData PrimitiveData;
+	FPrimitiveData PrimitiveData;
 	FMeshDescription Description;
 
 	FLightMapRef LightMap;
diff --git a/Engine/Plugins/Enterprise/GLTFExporter/Source/GLTFMaterialBaking/GLTFMaterialBaking.Build.cs b/Engine/Plugins/Enterprise/GLTFExporter/Source/GLTFMaterialBaking/GLTFMaterialBaking.Build.cs
deleted file mode 100644
diff --git a/Engine/Plugins/Enterprise/GLTFExporter/Source/GLTFMaterialBaking/Private/GLTFExportMaterialProxy.h b/Engine/Plugins/Enterprise/GLTFExporter/Source/GLTFMaterialBaking/Private/GLTFExportMaterialProxy.h
deleted file mode 100644
diff --git a/Engine/Plugins/Enterprise/GLTFExporter/Source/GLTFMaterialBaking/Private/GLTFMaterialBakingHelpers.cpp b/Engine/Plugins/Enterprise/GLTFExporter/Source/GLTFMaterialBaking/Private/GLTFMaterialBakingHelpers.cpp
deleted file mode 100644
diff --git a/Engine/Plugins/Enterprise/GLTFExporter/Source/GLTFMaterialBaking/Private/GLTFMaterialBakingModule.cpp b/Engine/Plugins/Enterprise/GLTFExporter/Source/GLTFMaterialBaking/Private/GLTFMaterialBakingModule.cpp
deleted file mode 100644
diff --git a/Engine/Plugins/Enterprise/GLTFExporter/Source/GLTFMaterialBaking/Private/GLTFMaterialBakingModule.h b/Engine/Plugins/Enterprise/GLTFExporter/Source/GLTFMaterialBaking/Private/GLTFMaterialBakingModule.h
deleted file mode 100644
diff --git a/Engine/Plugins/Enterprise/GLTFExporter/Source/GLTFMaterialBaking/Private/GLTFMaterialPropertyEx.cpp b/Engine/Plugins/Enterprise/GLTFExporter/Source/GLTFMaterialBaking/Private/GLTFMaterialPropertyEx.cpp
deleted file mode 100644
diff --git a/Engine/Plugins/Enterprise/GLTFExporter/Source/GLTFMaterialBaking/Private/MaterialRenderItem.cpp b/Engine/Plugins/Enterprise/GLTFExporter/Source/GLTFMaterialBaking/Private/MaterialRenderItem.cpp
deleted file mode 100644
diff --git a/Engine/Plugins/Enterprise/GLTFExporter/Source/GLTFMaterialBaking/Private/MaterialRenderItem.h b/Engine/Plugins/Enterprise/GLTFExporter/Source/GLTFMaterialBaking/Private/MaterialRenderItem.h
deleted file mode 100644
diff --git a/Engine/Plugins/Enterprise/GLTFExporter/Source/GLTFMaterialBaking/Private/MaterialRenderItemData.cpp b/Engine/Plugins/Enterprise/GLTFExporter/Source/GLTFMaterialBaking/Private/MaterialRenderItemData.cpp
deleted file mode 100644
diff --git a/Engine/Plugins/Enterprise/GLTFExporter/Source/GLTFMaterialBaking/Private/MaterialRenderItemData.h b/Engine/Plugins/Enterprise/GLTFExporter/Source/GLTFMaterialBaking/Private/MaterialRenderItemData.h
deleted file mode 100644
diff --git a/Engine/Plugins/Enterprise/GLTFExporter/Source/GLTFMaterialBaking/Public/GLTFMaterialBakingHelpers.h b/Engine/Plugins/Enterprise/GLTFExporter/Source/GLTFMaterialBaking/Public/GLTFMaterialBakingHelpers.h
deleted file mode 100644
diff --git a/Engine/Plugins/Enterprise/GLTFExporter/Source/GLTFMaterialBaking/Public/GLTFMaterialBakingStructures.h b/Engine/Plugins/Enterprise/GLTFExporter/Source/GLTFMaterialBaking/Public/GLTFMaterialBakingStructures.h
deleted file mode 100644
diff --git a/Engine/Plugins/Enterprise/GLTFExporter/Source/GLTFMaterialBaking/Public/GLTFMaterialPropertyEx.h b/Engine/Plugins/Enterprise/GLTFExporter/Source/GLTFMaterialBaking/Public/GLTFMaterialPropertyEx.h
deleted file mode 100644
diff --git a/Engine/Plugins/Enterprise/GLTFExporter/Source/GLTFMaterialBaking/Public/IGLTFMaterialBakingModule.h b/Engine/Plugins/Enterprise/GLTFExporter/Source/GLTFMaterialBaking/Public/IGLTFMaterialBakingModule.h
deleted file mode 100644
