Use engine's built-in material analysis in glTF exporter

diff --git a/Engine/Plugins/Enterprise/GLTFExporter/GLTFExporter.uplugin b/Engine/Plugins/Enterprise/GLTFExporter/GLTFExporter.uplugin
--- a/Engine/Plugins/Enterprise/GLTFExporter/GLTFExporter.uplugin
+++ b/Engine/Plugins/Enterprise/GLTFExporter/GLTFExporter.uplugin
@@ -17,12 +17,6 @@
 			"Type": "Runtime",
 			"LoadingPhase": "PostConfigInit",
 			"PlatformAllowList": [ "Win64", "Mac", "Linux" ]
-		},
-		{
-			"Name": "GLTFMaterialAnalyzer",
-			"Type": "Editor",
-			"LoadingPhase": "Default",
-			"PlatformAllowList": [ "Win64", "Mac", "Linux" ]
 		}
 	],
 	"Plugins": [
diff --git a/Engine/Plugins/Enterprise/GLTFExporter/Source/GLTFExporter/GLTFExporter.Build.cs b/Engine/Plugins/Enterprise/GLTFExporter/Source/GLTFExporter/GLTFExporter.Build.cs
--- a/Engine/Plugins/Enterprise/GLTFExporter/Source/GLTFExporter/GLTFExporter.Build.cs
+++ b/Engine/Plugins/Enterprise/GLTFExporter/Source/GLTFExporter/GLTFExporter.Build.cs
@@ -60,7 +60,6 @@ public class GLTFExporter : ModuleRules
 					"MeshMergeUtilities",
 					"MeshDescription",
 					"StaticMeshDescription",
-					"GLTFMaterialAnalyzer",
 				}
 			);
 
diff --git a/Engine/Plugins/Enterprise/GLTFExporter/Source/GLTFExporter/Private/Converters/GLTFMaterialUtilities.cpp b/Engine/Plugins/Enterprise/GLTFExporter/Source/GLTFExporter/Private/Converters/GLTFMaterialUtilities.cpp
--- a/Engine/Plugins/Enterprise/GLTFExporter/Source/GLTFExporter/Private/Converters/GLTFMaterialUtilities.cpp
+++ b/Engine/Plugins/Enterprise/GLTFExporter/Source/GLTFExporter/Private/Converters/GLTFMaterialUtilities.cpp
@@ -3,11 +3,11 @@
 #include "Converters/GLTFMaterialUtilities.h"
 #include "Converters/GLTFTextureUtilities.h"
 #include "Converters/GLTFNameUtilities.h"
-#include "Materials/Material.h"
+#include "Converters/GLTFProxyMaterialCompiler.h"
 #include "Utilities/GLTFProxyMaterialUtilities.h"
+#include "Materials/Material.h"
 #include "Misc/DefaultValueHelper.h"
 #if WITH_EDITOR
-#include "GLTFMaterialAnalyzer.h"
 #include "IMaterialBakingModule.h"
 #include "MaterialBakingStructures.h"
 #include "Materials/MaterialExpressionTextureSample.h"
@@ -95,11 +95,11 @@ const FExpressionInput* FGLTFMaterialUtilities::GetInputForProperty(const UMater
 	return UnderlyingMaterial->GetExpressionInputForProperty(Property.Type);
 }
 
-const UMaterialExpressionCustomOutput* FGLTFMaterialUtilities::GetCustomOutputByName(const UMaterialInterface* Material, const FString& FunctionName)
+UMaterialExpressionCustomOutput* FGLTFMaterialUtilities::GetCustomOutputByName(const UMaterialInterface* Material, const FString& FunctionName)
 {
 	for (const TObjectPtr<UMaterialExpression>& Expression : Material->GetMaterial()->GetExpressions())
 	{
-		const UMaterialExpressionCustomOutput* CustomOutput = Cast<UMaterialExpressionCustomOutput>(Expression);
+		UMaterialExpressionCustomOutput* CustomOutput = Cast<UMaterialExpressionCustomOutput>(Expression);
 		if (CustomOutput != nullptr && CustomOutput->GetFunctionName() == FunctionName)
 		{
 			return CustomOutput;
@@ -230,7 +230,7 @@ bool FGLTFMaterialUtilities::TryGetTextureCoordinateIndex(const UMaterialExpress
 
 void FGLTFMaterialUtilities::GetAllTextureCoordinateIndices(const UMaterialInterface* InMaterial, const FMaterialPropertyEx& InProperty, FGLTFIndexArray& OutTexCoords)
 {
-	FGLTFMaterialAnalysis Analysis;
+	FMaterialAnalysisResult Analysis;
 	AnalyzeMaterialProperty(InMaterial, InProperty, Analysis);
 
 	const TBitArray<>& TexCoords = Analysis.TextureCoordinates;
@@ -243,28 +243,53 @@ void FGLTFMaterialUtilities::GetAllTextureCoordinateIndices(const UMaterialInter
 	}
 }
 
-void FGLTFMaterialUtilities::AnalyzeMaterialProperty(const UMaterialInterface* InMaterial, const FMaterialPropertyEx& InProperty, FGLTFMaterialAnalysis& OutAnalysis)
+void FGLTFMaterialUtilities::AnalyzeMaterialProperty(const UMaterialInterface* InMaterial, const FMaterialPropertyEx& InProperty, FMaterialAnalysisResult& OutAnalysis)
 {
 	if (GetInputForProperty(InMaterial, InProperty) == nullptr)
 	{
-		OutAnalysis = FGLTFMaterialAnalysis();
+		OutAnalysis = FMaterialAnalysisResult();
 		return;
 	}
 
-	UGLTFMaterialAnalyzer::AnalyzeMaterialPropertyEx(InMaterial, InProperty.Type, InProperty.CustomOutput.ToString(), OutAnalysis);
+	UMaterial* BaseMaterial = const_cast<UMaterial*>(InMaterial->GetMaterial());
+	bool bRequiresPrimitiveData = false;
+
+	// To extend and improve the analysis for glTF's specific use-case, we compile using FGLTFProxyMaterialCompiler
+	BaseMaterial->AnalyzeMaterialCompilationInCallback([InProperty, BaseMaterial, &bRequiresPrimitiveData](FMaterialCompiler* Compiler)
+	{
+		FGLTFProxyMaterialCompiler ProxyCompiler(Compiler);
+
+		if (InProperty.IsCustomOutput())
+		{
+			UMaterialExpressionCustomOutput* CustomOutput = GetCustomOutputByName(BaseMaterial, InProperty.CustomOutput.ToString());
+			ProxyCompiler.SetMaterialProperty(MP_MAX, CustomOutput->GetShaderFrequency(), false);
+			CustomOutput->Compile(&ProxyCompiler, 0);
+		}
+		else
+		{
+			ProxyCompiler.SetMaterialProperty(InProperty.Type, SF_NumFrequencies, false);
+			BaseMaterial->CompileProperty(&ProxyCompiler, InProperty.Type);
+		}
+
+		bRequiresPrimitiveData =
+			ProxyCompiler.bUsesActorPosition ||
+			ProxyCompiler.bUsesObjectPosition ||
+			ProxyCompiler.bUsesObjectOrientation ||
+			ProxyCompiler.bUsesObjectRadius ||
+			ProxyCompiler.bUsesObjectBounds ||
+			ProxyCompiler.bUsesObjectLocalBounds ||
+			ProxyCompiler.bUsesPreSkinnedLocalBounds ||
+			ProxyCompiler.bUsesCustomPrimitiveData;
+	}, OutAnalysis);
+
+	// Also make sure the analysis takes into account primitive data
+	OutAnalysis.bRequiresVertexData |= bRequiresPrimitiveData;
 }
 
 FMaterialShadingModelField FGLTFMaterialUtilities::EvaluateShadingModelExpression(const UMaterialInterface* Material)
 {
-	FGLTFMaterialAnalysis Analysis;
+	FMaterialAnalysisResult Analysis;
 	AnalyzeMaterialProperty(Material, MP_ShadingModel, Analysis);
-
-	int32 Value;
-	if (FDefaultValueHelper::ParseInt(Analysis.ParameterCode, Value))
-	{
-		return static_cast<EMaterialShadingModel>(Value);
-	}
-
 	return Analysis.ShadingModels;
 }
 
@@ -334,7 +359,7 @@ bool FGLTFMaterialUtilities::NeedsMeshData(const UMaterialInterface* Material)
 		};
 
 		bool bNeedsMeshData = false;
-		FGLTFMaterialAnalysis Analysis;
+		FMaterialAnalysisResult Analysis;
 
 		// TODO: optimize baking by separating need for vertex data and primitive data
 
@@ -342,7 +367,6 @@ bool FGLTFMaterialUtilities::NeedsMeshData(const UMaterialInterface* Material)
 		{
 			AnalyzeMaterialProperty(Material, Property, Analysis);
 			bNeedsMeshData |= Analysis.bRequiresVertexData;
-			bNeedsMeshData |= Analysis.bRequiresPrimitiveData;
 		}
 
 		return bNeedsMeshData;
diff --git a/Engine/Plugins/Enterprise/GLTFExporter/Source/GLTFExporter/Private/Converters/GLTFMaterialUtilities.h b/Engine/Plugins/Enterprise/GLTFExporter/Source/GLTFExporter/Private/Converters/GLTFMaterialUtilities.h
--- a/Engine/Plugins/Enterprise/GLTFExporter/Source/GLTFExporter/Private/Converters/GLTFMaterialUtilities.h
+++ b/Engine/Plugins/Enterprise/GLTFExporter/Source/GLTFExporter/Private/Converters/GLTFMaterialUtilities.h
@@ -6,7 +6,7 @@
 #include "Converters/GLTFSharedArray.h"
 #include "MaterialExpressionIO.h"
 
-struct FGLTFMaterialAnalysis;
+struct FMaterialAnalysisResult;
 struct FMaterialPropertyEx;
 class UMaterialExpressionTextureSample;
 class UMaterialInstanceConstant;
@@ -54,7 +54,7 @@ struct FGLTFMaterialUtilities
 		return static_cast<const FMaterialInput<InputType>*>(ExpressionInput);
 	}
 
-	static const UMaterialExpressionCustomOutput* GetCustomOutputByName(const UMaterialInterface* Material, const FString& FunctionName);
+	static UMaterialExpressionCustomOutput* GetCustomOutputByName(const UMaterialInterface* Material, const FString& FunctionName);
 
 	static FGLTFPropertyBakeOutput BakeMaterialProperty(const FIntPoint& OutputSize, const FMaterialPropertyEx& Property, const UMaterialInterface* Material, int32 TexCoord, const FGLTFMeshData* MeshData = nullptr, const FGLTFIndexArray& MeshSectionIndices = {}, bool bFillAlpha = true, bool bAdjustNormalmaps = true);
 
@@ -66,7 +66,7 @@ struct FGLTFMaterialUtilities
 	static bool TryGetTextureCoordinateIndex(const UMaterialExpressionTextureSample* TextureSampler, int32& TexCoord, FGLTFJsonTextureTransform& Transform);
 	static void GetAllTextureCoordinateIndices(const UMaterialInterface* InMaterial, const FMaterialPropertyEx& InProperty, FGLTFIndexArray& OutTexCoords);
 
-	static void AnalyzeMaterialProperty(const UMaterialInterface* Material, const FMaterialPropertyEx& InProperty, FGLTFMaterialAnalysis& OutAnalysis);
+	static void AnalyzeMaterialProperty(const UMaterialInterface* Material, const FMaterialPropertyEx& InProperty, FMaterialAnalysisResult& OutAnalysis);
 
 	static FMaterialShadingModelField EvaluateShadingModelExpression(const UMaterialInterface* Material);
 #endif
diff --git a/Engine/Plugins/Enterprise/GLTFExporter/Source/GLTFExporter/Private/Converters/GLTFProxyMaterialCompiler.h b/Engine/Plugins/Enterprise/GLTFExporter/Source/GLTFExporter/Private/Converters/GLTFProxyMaterialCompiler.h
new file mode 100644
--- /dev/null
+++ b/Engine/Plugins/Enterprise/GLTFExporter/Source/GLTFExporter/Private/Converters/GLTFProxyMaterialCompiler.h
@@ -0,0 +1,160 @@
+// Copyright Epic Games, Inc. All Rights Reserved.
+
+#pragma once
+
+#if WITH_EDITOR
+
+#include "MaterialCompiler.h"
+#include "Materials/MaterialExpressionCustom.h"
+
+class FGLTFProxyMaterialCompiler : public FProxyMaterialCompiler
+{
+public:
+
+	bool bUsesActorPosition = false;
+	bool bUsesObjectPosition = false;
+	bool bUsesObjectOrientation = false;
+	bool bUsesObjectRadius = false;
+	bool bUsesObjectBounds = false;
+	bool bUsesObjectLocalBounds = false;
+	bool bUsesPreSkinnedLocalBounds = false;
+	bool bUsesCustomPrimitiveData = false;
+
+	using FProxyMaterialCompiler::FProxyMaterialCompiler;
+
+	virtual int32 ActorWorldPosition() override
+	{
+		bUsesActorPosition = true;
+		return Compiler->ActorWorldPosition();
+	}
+
+	virtual int32 ObjectWorldPosition() override
+	{
+		bUsesObjectPosition = true;
+		return Compiler->ObjectWorldPosition();
+	}
+
+	virtual int32 ObjectOrientation() override
+	{
+		bUsesObjectOrientation = true;
+		return Compiler->ObjectOrientation();
+	}
+
+	virtual int32 ObjectRadius() override
+	{
+		bUsesObjectRadius = true;
+		return Compiler->ObjectRadius();
+	}
+
+	virtual int32 ObjectBounds() override
+	{
+		bUsesObjectBounds = true;
+		return Compiler->ObjectBounds();
+	}
+
+	virtual int32 ObjectLocalBounds(int32 OutputIndex) override
+	{
+		bUsesObjectLocalBounds = true;
+		return Compiler->ObjectLocalBounds(OutputIndex);
+	}
+
+	// NOTE: also need to check old in-engine material function ObjectLocalBounds that bypasses the above override since it uses custom expressions 
+	virtual int32 CustomExpression(class UMaterialExpressionCustom* Custom, int32 OutputIndex, TArray<int32>& CompiledInputs) override
+	{
+		if (Custom->Code.Contains(TEXT("GetPrimitiveData(Parameters).LocalObjectBounds")))
+		{
+			bUsesObjectLocalBounds = true;
+		}
+
+		return Compiler->CustomExpression(Custom, OutputIndex, CompiledInputs);
+	}
+
+	virtual int32 PreSkinnedLocalBounds(int32 OutputIndex) override
+	{
+		bUsesPreSkinnedLocalBounds = true;
+		return Compiler->PreSkinnedLocalBounds(OutputIndex);
+	}
+
+	virtual int32 CustomPrimitiveData(int32 OutputIndex, EMaterialValueType Type) override
+	{
+		bUsesCustomPrimitiveData = true;
+		return Compiler->CustomPrimitiveData(OutputIndex, Type);
+	}
+
+	virtual int32 GIReplace(int32 Direct, int32 StaticIndirect, int32 DynamicIndirect) override
+	{
+		return Direct; // ignore all non-default branches in this runtime switch
+	}
+
+	virtual int32 ShadowReplace(int32 Default, int32 Shadow) override
+	{
+		return Default; // ignore all non-default branches in this runtime switch
+	}
+
+	virtual int32 RayTracingQualitySwitchReplace(int32 Normal, int32 RayTraced) override
+	{
+		return Normal; // ignore all non-default branches in this runtime switch
+	}
+
+	virtual int32 VirtualTextureOutputReplace(int32 Default, int32 VirtualTexture) override
+	{
+		return Default; // ignore all non-default branches in this runtime switch
+	}
+
+	virtual int32 ReflectionCapturePassSwitch(int32 Default, int32 Reflection) override
+	{
+		return Default; // ignore all non-default branches in this runtime switch
+	}
+
+	virtual int32 ReflectionAboutCustomWorldNormal(int32 CustomWorldNormal, int32 bNormalizeCustomWorldNormal) override
+	{
+		return Compiler->ReflectionAboutCustomWorldNormal(CustomWorldNormal, bNormalizeCustomWorldNormal);
+	}
+
+	virtual int32 ParticleRelativeTime() override
+	{
+		return Compiler->ParticleRelativeTime();
+	}
+
+	virtual int32 ParticleMotionBlurFade() override
+	{
+		return Compiler->ParticleMotionBlurFade();
+	}
+
+	virtual int32 ParticleRandom() override
+	{
+		return Compiler->ParticleRandom();
+	}
+
+	virtual int32 ParticleDirection() override
+	{
+		return Compiler->ParticleDirection();
+	}
+
+	virtual int32 ParticleSpeed() override
+	{
+		return Compiler->ParticleSpeed();
+	}
+
+	virtual int32 ParticleSize() override
+	{
+		return Compiler->ParticleSize();
+	}
+
+	virtual int32 VertexInterpolator(uint32 InterpolatorIndex) override
+	{
+		return Compiler->VertexInterpolator(InterpolatorIndex);
+	}
+
+	virtual int32 MaterialBakingWorldPosition() override
+	{
+		return Compiler->MaterialBakingWorldPosition();
+	}
+
+	virtual EMaterialCompilerType GetCompilerType() const override
+	{
+		return EMaterialCompilerType::MaterialProxy;
+	}
+};
+
+#endif
diff --git a/Engine/Plugins/Enterprise/GLTFExporter/Source/GLTFMaterialAnalyzer/GLTFMaterialAnalyzer.Build.cs b/Engine/Plugins/Enterprise/GLTFExporter/Source/GLTFMaterialAnalyzer/GLTFMaterialAnalyzer.Build.cs
deleted file mode 100644
diff --git a/Engine/Plugins/Enterprise/GLTFExporter/Source/GLTFMaterialAnalyzer/Private/GLTFMaterialAnalyzer.cpp b/Engine/Plugins/Enterprise/GLTFExporter/Source/GLTFMaterialAnalyzer/Private/GLTFMaterialAnalyzer.cpp
deleted file mode 100644
diff --git a/Engine/Plugins/Enterprise/GLTFExporter/Source/GLTFMaterialAnalyzer/Private/GLTFMaterialAnalyzerModule.cpp b/Engine/Plugins/Enterprise/GLTFExporter/Source/GLTFMaterialAnalyzer/Private/GLTFMaterialAnalyzerModule.cpp
deleted file mode 100644
diff --git a/Engine/Plugins/Enterprise/GLTFExporter/Source/GLTFMaterialAnalyzer/Private/GLTFProxyMaterialCompiler.h b/Engine/Plugins/Enterprise/GLTFExporter/Source/GLTFMaterialAnalyzer/Private/GLTFProxyMaterialCompiler.h
deleted file mode 100644
diff --git a/Engine/Plugins/Enterprise/GLTFExporter/Source/GLTFMaterialAnalyzer/Public/GLTFMaterialAnalysis.h b/Engine/Plugins/Enterprise/GLTFExporter/Source/GLTFMaterialAnalyzer/Public/GLTFMaterialAnalysis.h
deleted file mode 100644
diff --git a/Engine/Plugins/Enterprise/GLTFExporter/Source/GLTFMaterialAnalyzer/Public/GLTFMaterialAnalyzer.h b/Engine/Plugins/Enterprise/GLTFExporter/Source/GLTFMaterialAnalyzer/Public/GLTFMaterialAnalyzer.h
deleted file mode 100644
diff --git a/Engine/Plugins/Enterprise/GLTFExporter/Source/GLTFMaterialAnalyzer/Public/GLTFMaterialAnalyzerModule.h b/Engine/Plugins/Enterprise/GLTFExporter/Source/GLTFMaterialAnalyzer/Public/GLTFMaterialAnalyzerModule.h
deleted file mode 100644
