Remove UGLTFMaterialExporter::DefaultPreviewMesh since it's no longer in use.

diff --git a/Engine/Plugins/Enterprise/GLTFExporter/Source/GLTFExporter/Private/Exporters/GLTFMaterialExporter.cpp b/Engine/Plugins/Enterprise/GLTFExporter/Source/GLTFExporter/Private/Exporters/GLTFMaterialExporter.cpp
--- a/Engine/Plugins/Enterprise/GLTFExporter/Source/GLTFExporter/Private/Exporters/GLTFMaterialExporter.cpp
+++ b/Engine/Plugins/Enterprise/GLTFExporter/Source/GLTFExporter/Private/Exporters/GLTFMaterialExporter.cpp
@@ -13,8 +13,6 @@ UGLTFMaterialExporter::UGLTFMaterialExporter(const FObjectInitializer& ObjectIni
 	: Super(ObjectInitializer)
 {
 	SupportedClass = UMaterialInterface::StaticClass();
-
-	 // TODO: remove UGLTFMaterialExporter::DefaultPreviewMesh in 5.2 since it's no longer in use.
 }
 
 bool UGLTFMaterialExporter::AddObject(FGLTFContainerBuilder& Builder, const UObject* Object)
diff --git a/Engine/Plugins/Enterprise/GLTFExporter/Source/GLTFExporter/Public/Exporters/GLTFMaterialExporter.h b/Engine/Plugins/Enterprise/GLTFExporter/Source/GLTFExporter/Public/Exporters/GLTFMaterialExporter.h
--- a/Engine/Plugins/Enterprise/GLTFExporter/Source/GLTFExporter/Public/Exporters/GLTFMaterialExporter.h
+++ b/Engine/Plugins/Enterprise/GLTFExporter/Source/GLTFExporter/Public/Exporters/GLTFMaterialExporter.h
@@ -17,9 +17,4 @@ public:
 	explicit UGLTFMaterialExporter(const FObjectInitializer& ObjectInitializer = FObjectInitializer());
 
 	virtual bool AddObject(FGLTFContainerBuilder& Builder, const UObject* Object) override;
-
-private:
-
-	UPROPERTY(Transient, DuplicateTransient)
-	TObjectPtr<UStaticMesh> DefaultPreviewMesh;
 };
