Add copyright notice to gltf export

diff --git a/Engine/Plugins/Enterprise/GLTFExporter/Source/GLTFExporter/GLTFExporter.Build.cs b/Engine/Plugins/Enterprise/GLTFExporter/Source/GLTFExporter/GLTFExporter.Build.cs
--- a/Engine/Plugins/Enterprise/GLTFExporter/Source/GLTFExporter/GLTFExporter.Build.cs
+++ b/Engine/Plugins/Enterprise/GLTFExporter/Source/GLTFExporter/GLTFExporter.Build.cs
@@ -29,6 +29,7 @@ public class GLTFExporter : ModuleRules
 				"MovieSceneTracks",
 				"VariantManagerContent",
 				"Projects",
+				"EngineSettings",
 				"Analytics",
 			}
 		);
diff --git a/Engine/Plugins/Enterprise/GLTFExporter/Source/GLTFExporter/Private/Builders/GLTFJsonBuilder.cpp b/Engine/Plugins/Enterprise/GLTFExporter/Source/GLTFExporter/Private/Builders/GLTFJsonBuilder.cpp
--- a/Engine/Plugins/Enterprise/GLTFExporter/Source/GLTFExporter/Private/Builders/GLTFJsonBuilder.cpp
+++ b/Engine/Plugins/Enterprise/GLTFExporter/Source/GLTFExporter/Private/Builders/GLTFJsonBuilder.cpp
@@ -3,12 +3,14 @@
 #include "Builders/GLTFJsonBuilder.h"
 #include "Interfaces/IPluginManager.h"
 #include "Runtime/Launch/Resources/Version.h"
+#include "GeneralProjectSettings.h"
 
 FGLTFJsonBuilder::FGLTFJsonBuilder(const FString& FileName, const UGLTFExportOptions* ExportOptions)
 	: FGLTFFileBuilder(FileName, ExportOptions)
 	, DefaultScene(JsonRoot.DefaultScene)
 {
 	JsonRoot.Asset.Generator = GetGeneratorString();
+	JsonRoot.Asset.Copyright = GetCopyrightString();
 }
 
 void FGLTFJsonBuilder::WriteJsonArchive(FArchive& Archive)
@@ -105,7 +107,13 @@ const FGLTFJsonRoot& FGLTFJsonBuilder::GetRoot() const
 	return JsonRoot;
 }
 
-const TCHAR* FGLTFJsonBuilder::GetGeneratorString()
+FString FGLTFJsonBuilder::GetGeneratorString()
 {
 	return TEXT(EPIC_PRODUCT_NAME) TEXT(" ") ENGINE_VERSION_STRING;
 }
+
+FString FGLTFJsonBuilder::GetCopyrightString()
+{
+	FString CopyrightNotice = GetDefault<UGeneralProjectSettings>()->CopyrightNotice;
+	return CopyrightNotice != TEXT("Fill out your copyright notice in the Description page of Project Settings.") ? CopyrightNotice : TEXT("");
+}
diff --git a/Engine/Plugins/Enterprise/GLTFExporter/Source/GLTFExporter/Public/Builders/GLTFJsonBuilder.h b/Engine/Plugins/Enterprise/GLTFExporter/Source/GLTFExporter/Public/Builders/GLTFJsonBuilder.h
--- a/Engine/Plugins/Enterprise/GLTFExporter/Source/GLTFExporter/Public/Builders/GLTFJsonBuilder.h
+++ b/Engine/Plugins/Enterprise/GLTFExporter/Source/GLTFExporter/Public/Builders/GLTFJsonBuilder.h
@@ -39,7 +39,8 @@ protected:
 
 private:
 
-	static const TCHAR* GetGeneratorString();
+	static FString GetGeneratorString();
+	static FString GetCopyrightString();
 
 	FGLTFJsonRoot JsonRoot;
 };
