add UGLTFExportOptions::bIncludeCopyrightNotice to determine whether copyright is included in exported gltf. Also "hide" the option and bSkipNearDefaultValues under Advanced.

diff --git a/Engine/Plugins/Enterprise/GLTFExporter/Source/GLTFExporter/Private/Builders/GLTFJsonBuilder.cpp b/Engine/Plugins/Enterprise/GLTFExporter/Source/GLTFExporter/Private/Builders/GLTFJsonBuilder.cpp
--- a/Engine/Plugins/Enterprise/GLTFExporter/Source/GLTFExporter/Private/Builders/GLTFJsonBuilder.cpp
+++ b/Engine/Plugins/Enterprise/GLTFExporter/Source/GLTFExporter/Private/Builders/GLTFJsonBuilder.cpp
@@ -10,7 +10,10 @@ FGLTFJsonBuilder::FGLTFJsonBuilder(const FString& FileName, const UGLTFExportOpt
 	, DefaultScene(JsonRoot.DefaultScene)
 {
 	JsonRoot.Asset.Generator = GetGeneratorString();
-	JsonRoot.Asset.Copyright = GetCopyrightString();
+	if (ExportOptions->bIncludeCopyrightNotice)
+	{
+		JsonRoot.Asset.Copyright = GetCopyrightString();
+	}
 }
 
 void FGLTFJsonBuilder::WriteJsonArchive(FArchive& Archive)
@@ -114,6 +117,5 @@ FString FGLTFJsonBuilder::GetGeneratorString()
 
 FString FGLTFJsonBuilder::GetCopyrightString()
 {
-	FString CopyrightNotice = GetDefault<UGeneralProjectSettings>()->CopyrightNotice;
-	return CopyrightNotice != TEXT("Fill out your copyright notice in the Description page of Project Settings.") ? CopyrightNotice : TEXT("");
+	return GetDefault<UGeneralProjectSettings>()->CopyrightNotice;
 }
diff --git a/Engine/Plugins/Enterprise/GLTFExporter/Source/GLTFExporter/Private/Options/GLTFExportOptions.cpp b/Engine/Plugins/Enterprise/GLTFExporter/Source/GLTFExporter/Private/Options/GLTFExportOptions.cpp
--- a/Engine/Plugins/Enterprise/GLTFExporter/Source/GLTFExporter/Private/Options/GLTFExportOptions.cpp
+++ b/Engine/Plugins/Enterprise/GLTFExporter/Source/GLTFExporter/Private/Options/GLTFExportOptions.cpp
@@ -13,6 +13,7 @@ void UGLTFExportOptions::ResetToDefault()
 	ExportUniformScale = 0.01;
 	bExportPreviewMesh = true;
 	bSkipNearDefaultValues = true;
+	bIncludeCopyrightNotice = false;
 	bExportProxyMaterials = true;
 	bExportUnlitMaterials = true;
 	bExportClearCoatMaterials = true;
diff --git a/Engine/Plugins/Enterprise/GLTFExporter/Source/GLTFExporter/Public/Options/GLTFExportOptions.h b/Engine/Plugins/Enterprise/GLTFExporter/Source/GLTFExporter/Public/Options/GLTFExportOptions.h
--- a/Engine/Plugins/Enterprise/GLTFExporter/Source/GLTFExporter/Public/Options/GLTFExportOptions.h
+++ b/Engine/Plugins/Enterprise/GLTFExporter/Source/GLTFExporter/Public/Options/GLTFExportOptions.h
@@ -42,9 +42,13 @@ class GLTFEXPORTER_API UGLTFExportOptions : public UObject
 	bool bExportPreviewMesh;
 
 	/** If enabled, floating-point-based JSON properties that are nearly equal to their default value will not be exported and thus regarded as exactly default, reducing JSON size. */
-	UPROPERTY(EditAnywhere, BlueprintReadWrite, Config, Category = General)
+	UPROPERTY(EditAnywhere, BlueprintReadWrite, Config, Category = General, AdvancedDisplay)
 	bool bSkipNearDefaultValues;
 
+	/** If enabled, the copyright notice from project settings will be included as metadata in the glTF asset. */
+	UPROPERTY(EditAnywhere, BlueprintReadWrite, Config, Category = General, AdvancedDisplay)
+	bool bIncludeCopyrightNotice;
+
 	/** If enabled, materials that have a proxy defined in their user data, will be exported using that proxy instead. This setting won't affect proxy materials exported or referenced directly. */
 	UPROPERTY(EditAnywhere, BlueprintReadWrite, Config, Category = Material)
 	bool bExportProxyMaterials;
