Fix truncation warnings in GLTFExporter plugin

diff --git a/Engine/Plugins/Enterprise/GLTFExporter/Source/GLTFExporter/Private/Builders/GLTFContainerBuilder.cpp b/Engine/Plugins/Enterprise/GLTFExporter/Source/GLTFExporter/Private/Builders/GLTFContainerBuilder.cpp
--- a/Engine/Plugins/Enterprise/GLTFExporter/Source/GLTFExporter/Private/Builders/GLTFContainerBuilder.cpp
+++ b/Engine/Plugins/Enterprise/GLTFExporter/Source/GLTFExporter/Private/Builders/GLTFContainerBuilder.cpp
@@ -116,20 +116,21 @@ void FGLTFContainerBuilder::WriteData(FArchive& Archive, const TArray64<uint8>&
 	Archive.Serialize(const_cast<uint8*>(Data.GetData()), Data.Num());
 }
 
-void FGLTFContainerBuilder::WriteFill(FArchive& Archive, int32 Size, uint8 Value)
+void FGLTFContainerBuilder::WriteFill(FArchive& Archive, uint32 Size, uint8 Value)
 {
-	while (--Size >= 0)
+	while (Size != 0)
 	{
+		Size--;
 		Archive.Serialize(&Value, sizeof(Value));
 	}
 }
 
-int32 FGLTFContainerBuilder::GetPaddedChunkSize(int32 Size)
+uint32 FGLTFContainerBuilder::GetPaddedChunkSize(uint64 Size)
 {
-	return (Size + 3) & ~3;
+	return static_cast<uint32>((Size + 3) & ~3);
 }
 
-int32 FGLTFContainerBuilder::GetTrailingChunkSize(int32 Size)
+uint32 FGLTFContainerBuilder::GetTrailingChunkSize(uint64 Size)
 {
-	return (4 - (Size & 3)) & 3;
+	return static_cast<uint32>((4 - (Size & 3)) & 3);
 }
diff --git a/Engine/Plugins/Enterprise/GLTFExporter/Source/GLTFExporter/Private/Builders/GLTFTaskBuilder.cpp b/Engine/Plugins/Enterprise/GLTFExporter/Source/GLTFExporter/Private/Builders/GLTFTaskBuilder.cpp
--- a/Engine/Plugins/Enterprise/GLTFExporter/Source/GLTFExporter/Private/Builders/GLTFTaskBuilder.cpp
+++ b/Engine/Plugins/Enterprise/GLTFExporter/Source/GLTFExporter/Private/Builders/GLTFTaskBuilder.cpp
@@ -27,7 +27,8 @@ void FGLTFTaskBuilder::ProcessSlowTasks(FFeedbackContext* Context)
 		if (Context != nullptr)
 		{
 			const FText MessageFormat = GetPriorityMessageFormat(Priority);
-			FScopedSlowTask Progress(Tasks->Num(), FText::Format(MessageFormat, FText()), true, *Context);
+			const float AmountOfWork = static_cast<float>(Tasks->Num());
+			FScopedSlowTask Progress(AmountOfWork, FText::Format(MessageFormat, FText()), true, *Context);
 			Progress.MakeDialog();
 
 			for (TUniquePtr<FGLTFDelayedTask>& Task : *Tasks)
diff --git a/Engine/Plugins/Enterprise/GLTFExporter/Source/GLTFExporter/Private/Converters/GLTFBoneUtilities.cpp b/Engine/Plugins/Enterprise/GLTFExporter/Source/GLTFExporter/Private/Converters/GLTFBoneUtilities.cpp
--- a/Engine/Plugins/Enterprise/GLTFExporter/Source/GLTFExporter/Private/Converters/GLTFBoneUtilities.cpp
+++ b/Engine/Plugins/Enterprise/GLTFExporter/Source/GLTFExporter/Private/Converters/GLTFBoneUtilities.cpp
@@ -46,7 +46,7 @@ void FGLTFBoneUtilities::GetBoneIndices(const USkeleton* Skeleton, TArray<FBoneI
 
 	for (int32 BoneIndex = 0; BoneIndex < BoneCount; ++BoneIndex)
 	{
-		OutBoneIndices[BoneIndex] = BoneIndex;
+		OutBoneIndices[BoneIndex] = static_cast<FBoneIndexType>(BoneIndex);
 	}
 }
 
diff --git a/Engine/Plugins/Enterprise/GLTFExporter/Source/GLTFExporter/Private/Converters/GLTFMaterialProxyFactory.cpp b/Engine/Plugins/Enterprise/GLTFExporter/Source/GLTFExporter/Private/Converters/GLTFMaterialProxyFactory.cpp
--- a/Engine/Plugins/Enterprise/GLTFExporter/Source/GLTFExporter/Private/Converters/GLTFMaterialProxyFactory.cpp
+++ b/Engine/Plugins/Enterprise/GLTFExporter/Source/GLTFExporter/Private/Converters/GLTFMaterialProxyFactory.cpp
@@ -155,7 +155,7 @@ void FGLTFMaterialProxyFactory::SetProxyParameter(UMaterialInstanceConstant* Pro
 	}
 
 	ParameterInfo.Texture.Set(ProxyMaterial, Texture, true);
-	ParameterInfo.UVIndex.Set(ProxyMaterial, TextureInfo.TexCoord, true);
+	ParameterInfo.UVIndex.Set(ProxyMaterial, static_cast<float>(TextureInfo.TexCoord), true);
 	ParameterInfo.UVOffset.Set(ProxyMaterial, FLinearColor(TextureInfo.Transform.Offset.X, TextureInfo.Transform.Offset.Y, 0, 0), true);
 	ParameterInfo.UVScale.Set(ProxyMaterial, FLinearColor(TextureInfo.Transform.Scale.X, TextureInfo.Transform.Scale.Y, 0, 0), true);
 	ParameterInfo.UVRotation.Set(ProxyMaterial, TextureInfo.Transform.Rotation, true);
diff --git a/Engine/Plugins/Enterprise/GLTFExporter/Source/GLTFExporter/Private/Converters/GLTFMaterialUtilities.cpp b/Engine/Plugins/Enterprise/GLTFExporter/Source/GLTFExporter/Private/Converters/GLTFMaterialUtilities.cpp
--- a/Engine/Plugins/Enterprise/GLTFExporter/Source/GLTFExporter/Private/Converters/GLTFMaterialUtilities.cpp
+++ b/Engine/Plugins/Enterprise/GLTFExporter/Source/GLTFExporter/Private/Converters/GLTFMaterialUtilities.cpp
@@ -194,7 +194,12 @@ FGLTFJsonTexture* FGLTFMaterialUtilities::AddTexture(FGLTFConvertBuilder& Builde
 
 FLinearColor FGLTFMaterialUtilities::GetMask(const FExpressionInput& ExpressionInput)
 {
-	return FLinearColor(ExpressionInput.MaskR, ExpressionInput.MaskG, ExpressionInput.MaskB, ExpressionInput.MaskA);
+	return FLinearColor(
+		static_cast<float>(ExpressionInput.MaskR),
+		static_cast<float>(ExpressionInput.MaskG),
+		static_cast<float>(ExpressionInput.MaskB),
+		static_cast<float>(ExpressionInput.MaskA)
+	);
 }
 
 uint32 FGLTFMaterialUtilities::GetMaskComponentCount(const FExpressionInput& ExpressionInput)
@@ -215,10 +220,10 @@ bool FGLTFMaterialUtilities::TryGetTextureCoordinateIndex(const UMaterialExpress
 	if (const UMaterialExpressionTextureCoordinate* TextureCoordinate = Cast<UMaterialExpressionTextureCoordinate>(Expression))
 	{
 		TexCoord = TextureCoordinate->CoordinateIndex;
-		Transform.Offset.X = TextureCoordinate->UnMirrorU ? TextureCoordinate->UTiling * 0.5 : 0.0;
-		Transform.Offset.Y = TextureCoordinate->UnMirrorV ? TextureCoordinate->VTiling * 0.5 : 0.0;
-		Transform.Scale.X = TextureCoordinate->UTiling * (TextureCoordinate->UnMirrorU ? 0.5 : 1.0);
-		Transform.Scale.Y = TextureCoordinate->VTiling * (TextureCoordinate->UnMirrorV ? 0.5 : 1.0);
+		Transform.Offset.X = TextureCoordinate->UnMirrorU ? TextureCoordinate->UTiling * 0.5f : 0.0f;
+		Transform.Offset.Y = TextureCoordinate->UnMirrorV ? TextureCoordinate->VTiling * 0.5f : 0.0f;
+		Transform.Scale.X = TextureCoordinate->UTiling * (TextureCoordinate->UnMirrorU ? 0.5f : 1.0f);
+		Transform.Scale.Y = TextureCoordinate->VTiling * (TextureCoordinate->UnMirrorV ? 0.5f : 1.0f);
 		Transform.Rotation = 0;
 		return true;
 	}
diff --git a/Engine/Plugins/Enterprise/GLTFExporter/Source/GLTFExporter/Private/Converters/GLTFTextureUtilities.cpp b/Engine/Plugins/Enterprise/GLTFExporter/Source/GLTFExporter/Private/Converters/GLTFTextureUtilities.cpp
--- a/Engine/Plugins/Enterprise/GLTFExporter/Source/GLTFExporter/Private/Converters/GLTFTextureUtilities.cpp
+++ b/Engine/Plugins/Enterprise/GLTFExporter/Source/GLTFExporter/Private/Converters/GLTFTextureUtilities.cpp
@@ -93,8 +93,8 @@ int32 FGLTFTextureUtilities::GetMipBias(const UTexture* Texture)
 
 FIntPoint FGLTFTextureUtilities::GetInGameSize(const UTexture* Texture)
 {
-	const int32 Width = Texture->GetSurfaceWidth();
-	const int32 Height = Texture->GetSurfaceHeight();
+	const int32 Width = static_cast<int32>(Texture->GetSurfaceWidth());
+	const int32 Height = static_cast<int32>(Texture->GetSurfaceHeight());
 
 	const int32 MipBias = GetMipBias(Texture);
 
diff --git a/Engine/Plugins/Enterprise/GLTFExporter/Source/GLTFExporter/Private/Tasks/GLTFDelayedAnimationTasks.cpp b/Engine/Plugins/Enterprise/GLTFExporter/Source/GLTFExporter/Private/Tasks/GLTFDelayedAnimationTasks.cpp
--- a/Engine/Plugins/Enterprise/GLTFExporter/Source/GLTFExporter/Private/Tasks/GLTFDelayedAnimationTasks.cpp
+++ b/Engine/Plugins/Enterprise/GLTFExporter/Source/GLTFExporter/Private/Tasks/GLTFDelayedAnimationTasks.cpp
@@ -166,7 +166,7 @@ void FGLTFDelayedLevelSequenceTask::Process()
 
 	for (int32 Frame = 0; Frame < FrameCount; ++Frame)
 	{
-		Timestamps[Frame] = DisplayRate.AsSeconds(Frame);
+		Timestamps[Frame] = static_cast<float>(DisplayRate.AsSeconds(Frame));
 		FrameTimes[Frame] = FFrameRate::TransformTime(FFrameTime(FrameOffset + Frame), DisplayRate, TickResolution);
 	}
 
diff --git a/Engine/Plugins/Enterprise/GLTFExporter/Source/GLTFExporter/Private/Tasks/GLTFDelayedMeshTasks.cpp b/Engine/Plugins/Enterprise/GLTFExporter/Source/GLTFExporter/Private/Tasks/GLTFDelayedMeshTasks.cpp
--- a/Engine/Plugins/Enterprise/GLTFExporter/Source/GLTFExporter/Private/Tasks/GLTFDelayedMeshTasks.cpp
+++ b/Engine/Plugins/Enterprise/GLTFExporter/Source/GLTFExporter/Private/Tasks/GLTFDelayedMeshTasks.cpp
@@ -272,9 +272,9 @@ void FGLTFDelayedSkeletalMeshTask::Process()
 	ValidateVertexBuffer(Builder, &VertexBuffer, *SkeletalMesh->GetName());
 
 	const TArray<FSkeletalMaterial>& MaterialSlots = FGLTFMeshUtilities::GetMaterials(SkeletalMesh);
-	const uint16 MaterialCount = MaterialSlots.Num();
+	const int32 MaterialCount = MaterialSlots.Num();
 
-	for (uint16 MaterialIndex = 0; MaterialIndex < MaterialCount; ++MaterialIndex)
+	for (int32 MaterialIndex = 0; MaterialIndex < MaterialCount; ++MaterialIndex)
 	{
 		const FGLTFIndexArray SectionIndices = FGLTFMeshUtilities::GetSectionIndices(MeshLOD, MaterialIndex);
 		const FGLTFMeshSection* ConvertedSection = MeshSectionConverter.GetOrAdd(&MeshLOD, SectionIndices);
diff --git a/Engine/Plugins/Enterprise/GLTFExporter/Source/GLTFExporter/Private/UI/GLTFExportOptionsWindow.cpp b/Engine/Plugins/Enterprise/GLTFExporter/Source/GLTFExporter/Private/UI/GLTFExportOptionsWindow.cpp
--- a/Engine/Plugins/Enterprise/GLTFExporter/Source/GLTFExporter/Private/UI/GLTFExportOptionsWindow.cpp
+++ b/Engine/Plugins/Enterprise/GLTFExporter/Source/GLTFExporter/Private/UI/GLTFExportOptionsWindow.cpp
@@ -43,16 +43,16 @@ void SGLTFExportOptionsWindow::Construct(const FArguments& InArgs)
 			SNew(SVerticalBox)
 			+ SVerticalBox::Slot()
 			.AutoHeight()
-			.Padding(2)
+			.Padding(2.0f)
 			[
 				SAssignNew(HeaderToolBox, SBox)
 			]
 			+SVerticalBox::Slot()
 			.AutoHeight()
-			.Padding(2)
+			.Padding(2.0f)
 			[
 				SNew(SBorder)
-				.Padding(FMargin(3))
+				.Padding(FMargin(3.0f))
 				.BorderImage(FAppStyle::GetBrush("ToolPanel.GroupBorder"))
 				[
 					SNew(SHorizontalBox)
@@ -64,7 +64,7 @@ void SGLTFExportOptionsWindow::Construct(const FArguments& InArgs)
 						.Text(LOCTEXT("Export_CurrentFileTitle", "Current File: "))
 					]
 					+SHorizontalBox::Slot()
-					.Padding(5, 0, 0, 0)
+					.Padding(5.0f, 0.0f, 0.0f, 0.0f)
 					.AutoWidth()
 					.VAlign(VAlign_Center)
 					[
@@ -76,17 +76,17 @@ void SGLTFExportOptionsWindow::Construct(const FArguments& InArgs)
 			]
 			+ SVerticalBox::Slot()
 			.FillHeight(1.0f)
-			.Padding(2)
+			.Padding(2.0f)
 			[
 				SAssignNew(InspectorBox, SBox)
 			]
 			+ SVerticalBox::Slot()
 			.AutoHeight()
 			.HAlign(HAlign_Right)
-			.Padding(2)
+			.Padding(2.0f)
 			[
 				SNew(SUniformGridPanel)
-				.SlotPadding(2)
+				.SlotPadding(2.0f)
 				+ SUniformGridPanel::Slot(0, 0)
 				[
 					SNew(SButton)
diff --git a/Engine/Plugins/Enterprise/GLTFExporter/Source/GLTFExporter/Private/UI/GLTFProxyOptionsWindow.cpp b/Engine/Plugins/Enterprise/GLTFExporter/Source/GLTFExporter/Private/UI/GLTFProxyOptionsWindow.cpp
--- a/Engine/Plugins/Enterprise/GLTFExporter/Source/GLTFExporter/Private/UI/GLTFProxyOptionsWindow.cpp
+++ b/Engine/Plugins/Enterprise/GLTFExporter/Source/GLTFExporter/Private/UI/GLTFProxyOptionsWindow.cpp
@@ -33,17 +33,17 @@ void SGLTFProxyOptionsWindow::Construct(const FArguments& InArgs)
 			SNew(SVerticalBox)
 			+ SVerticalBox::Slot()
 			.FillHeight(1.0f)
-			.Padding(2)
+			.Padding(2.0f)
 			[
 				SAssignNew(InspectorBox, SBox)
 			]
 			+ SVerticalBox::Slot()
 			.AutoHeight()
 			.HAlign(HAlign_Right)
-			.Padding(2)
+			.Padding(2.0f)
 			[
 				SNew(SUniformGridPanel)
-				.SlotPadding(2)
+				.SlotPadding(2.0f)
 				+ SUniformGridPanel::Slot(0, 0)
 				[
 					SAssignNew(ConfirmButton, SButton)
diff --git a/Engine/Plugins/Enterprise/GLTFExporter/Source/GLTFExporter/Public/Builders/GLTFContainerBuilder.h b/Engine/Plugins/Enterprise/GLTFExporter/Source/GLTFExporter/Public/Builders/GLTFContainerBuilder.h
--- a/Engine/Plugins/Enterprise/GLTFExporter/Source/GLTFExporter/Public/Builders/GLTFContainerBuilder.h
+++ b/Engine/Plugins/Enterprise/GLTFExporter/Source/GLTFExporter/Public/Builders/GLTFContainerBuilder.h
@@ -30,8 +30,8 @@ private:
 
 	static void WriteInt(FArchive& Archive, uint32 Value);
 	static void WriteData(FArchive& Archive, const TArray64<uint8>& Data);
-	static void WriteFill(FArchive& Archive, int32 Size, uint8 Value);
+	static void WriteFill(FArchive& Archive, uint32 Size, uint8 Value);
 
-	static int32 GetPaddedChunkSize(int32 Size);
-	static int32 GetTrailingChunkSize(int32 Size);
+	static uint32 GetPaddedChunkSize(uint64 Size);
+	static uint32 GetTrailingChunkSize(uint64 Size);
 };
