Fix URIs not being percent-encoded (and sanitized) in glTF export

diff --git a/Engine/Plugins/Enterprise/GLTFExporter/Source/GLTFExporter/Private/Builders/GLTFFileBuilder.cpp b/Engine/Plugins/Enterprise/GLTFExporter/Source/GLTFExporter/Private/Builders/GLTFFileBuilder.cpp
--- a/Engine/Plugins/Enterprise/GLTFExporter/Source/GLTFExporter/Private/Builders/GLTFFileBuilder.cpp
+++ b/Engine/Plugins/Enterprise/GLTFExporter/Source/GLTFExporter/Private/Builders/GLTFFileBuilder.cpp
@@ -9,11 +9,11 @@ FGLTFFileBuilder::FGLTFFileBuilder(const FString& FileName, const UGLTFExportOpt
 {
 }
 
-FString FGLTFFileBuilder::AddExternalFile(const FString& URI, const TSharedPtr<FGLTFMemoryArchive>& Archive)
+FString FGLTFFileBuilder::AddExternalFile(const FString& DesiredURI, const TSharedPtr<FGLTFMemoryArchive>& Archive)
 {
-	const FString UnqiueURI = GetUniqueURI(URI);
-	ExternalArchives.Add(UnqiueURI, Archive);
-	return UnqiueURI;
+	const FString ActualFileName = GetUniqueFileName(SanitizeFileName(DesiredURI));
+	ExternalArchives.Add(ActualFileName, Archive);
+	return EncodeURI(ActualFileName);
 }
 
 void FGLTFFileBuilder::GetExternalFiles(TArray<FString>& OutFilePaths, const FString& DirPath) const
@@ -46,26 +46,84 @@ bool FGLTFFileBuilder::WriteExternalFiles(const FString& DirPath, uint32 WriteFl
 	return true;
 }
 
-FString FGLTFFileBuilder::GetUniqueURI(const FString& URI) const
+FString FGLTFFileBuilder::GetUniqueFileName(const FString& InFileName) const
 {
-	if (!ExternalArchives.Contains(URI))
+	if (!ExternalArchives.Contains(InFileName))
 	{
-		return URI;
+		return InFileName;
 	}
 
-	const FString BaseFilename = FPaths::GetBaseFilename(URI);
-	const FString FileExtension = FPaths::GetExtension(URI, true);
-	FString UnqiueURI;
+	const FString BaseFileName = FPaths::GetBaseFilename(InFileName);
+	const FString FileExtension = FPaths::GetExtension(InFileName, true);
+	FString UniqueFileName;
 
 	int32 Suffix = 1;
 	do
 	{
-		UnqiueURI = BaseFilename + TEXT('_') + FString::FromInt(Suffix) + FileExtension;
+		UniqueFileName = BaseFileName + TEXT('_') + FString::FromInt(Suffix) + FileExtension;
 		Suffix++;
 	}
-	while (ExternalArchives.Contains(UnqiueURI));
+	while (ExternalArchives.Contains(UniqueFileName));
 
-	return UnqiueURI;
+	return UniqueFileName;
+}
+
+FString FGLTFFileBuilder::SanitizeFileName(const FString& InFileName)
+{
+	static const TCHAR* IllegalChars = TEXT("<>:\"/\\|?*");
+
+	FString SanitizedFileName;
+	SanitizedFileName.Reserve(InFileName.Len());
+
+	for (const TCHAR& Chr : InFileName)
+	{
+		if (FCString::Strchr(IllegalChars, Chr) == nullptr)
+		{
+			SanitizedFileName += Chr;
+		}
+	}
+
+	return SanitizedFileName;
+}
+
+FString FGLTFFileBuilder::EncodeURI(const FString& InFileName)
+{
+	// glTF only strictly requires the following reserved characters to be percent-encoded
+	static const TCHAR* CharEncodingMap[][2] =
+	{
+		// Always replace "%" first to avoid double-encoding characters
+		{ TEXT("%"), TEXT("%25") },
+		{ TEXT(":"), TEXT("%3A") },
+		{ TEXT("/"), TEXT("%2F") },
+		{ TEXT("?"), TEXT("%3F") },
+		{ TEXT("#"), TEXT("%23") },
+		{ TEXT("["), TEXT("%5B") },
+		{ TEXT("]"), TEXT("%5D") },
+		{ TEXT("@"), TEXT("%40") },
+		{ TEXT("!"), TEXT("%21") },
+		{ TEXT("$"), TEXT("%24") },
+		{ TEXT("&"), TEXT("%26") },
+		{ TEXT("'"), TEXT("%27") },
+		{ TEXT("("), TEXT("%28") },
+		{ TEXT(")"), TEXT("%29") },
+		{ TEXT("*"), TEXT("%2A") },
+		{ TEXT("+"), TEXT("%2B") },
+		{ TEXT(","), TEXT("%2C") },
+		{ TEXT(";"), TEXT("%3B") },
+		{ TEXT("="), TEXT("%3D") },
+	};
+
+	FString EncodedURI = InFileName;
+
+	for (int32 Index = 0; Index < GetNum(CharEncodingMap); Index++ )
+	{
+		const TCHAR* OriginalChar = CharEncodingMap[Index][0];
+		const TCHAR* EncodedChar = CharEncodingMap[Index][1];
+		// Use ReplaceInline as that won't create a copy of the string if the character isn't found
+		EncodedURI.ReplaceInline(OriginalChar, EncodedChar, ESearchCase::CaseSensitive);
+	}
+
+	return EncodedURI;
 }
 
 bool FGLTFFileBuilder::SaveToFile(const FString& FilePath, const TArray64<uint8>& FileData, uint32 WriteFlags)
diff --git a/Engine/Plugins/Enterprise/GLTFExporter/Source/GLTFExporter/Public/Builders/GLTFFileBuilder.h b/Engine/Plugins/Enterprise/GLTFExporter/Source/GLTFExporter/Public/Builders/GLTFFileBuilder.h
--- a/Engine/Plugins/Enterprise/GLTFExporter/Source/GLTFExporter/Public/Builders/GLTFFileBuilder.h
+++ b/Engine/Plugins/Enterprise/GLTFExporter/Source/GLTFExporter/Public/Builders/GLTFFileBuilder.h
@@ -11,7 +11,7 @@ public:
 
 	FGLTFFileBuilder(const FString& FileName, const UGLTFExportOptions* ExportOptions = nullptr);
 
-	FString AddExternalFile(const FString& URI, const TSharedPtr<FGLTFMemoryArchive>& Archive = MakeShared<FGLTFMemoryArchive>());
+	FString AddExternalFile(const FString& DesiredURI, const TSharedPtr<FGLTFMemoryArchive>& Archive);
 
 	void GetExternalFiles(TArray<FString>& OutFilePaths, const FString& DirPath = TEXT("")) const;
 
@@ -23,7 +23,11 @@ private:
 
 	TMap<FString, TSharedPtr<FGLTFMemoryArchive>> ExternalArchives;
 
-	FString GetUniqueURI(const FString& URI) const;
+	FString GetUniqueFileName(const FString& InFileName) const;
+
+	static FString SanitizeFileName(const FString& InFileName);
+
+	static FString EncodeURI(const FString& InFileName);
 
 protected:
 
