Fix glTF exporter to correctly material analyze all TransformVector and TransformPosition expressions

diff --git a/Engine/Plugins/Enterprise/GLTFExporter/Source/GLTFExporter/Private/Converters/GLTFMaterialUtilities.cpp b/Engine/Plugins/Enterprise/GLTFExporter/Source/GLTFExporter/Private/Converters/GLTFMaterialUtilities.cpp
--- a/Engine/Plugins/Enterprise/GLTFExporter/Source/GLTFExporter/Private/Converters/GLTFMaterialUtilities.cpp
+++ b/Engine/Plugins/Enterprise/GLTFExporter/Source/GLTFExporter/Private/Converters/GLTFMaterialUtilities.cpp
@@ -276,15 +276,7 @@ void FGLTFMaterialUtilities::AnalyzeMaterialProperty(const UMaterialInterface* I
 			BaseMaterial->CompileProperty(&ProxyCompiler, InProperty.Type);
 		}
 
-		bRequiresPrimitiveData =
-			ProxyCompiler.bUsesActorPosition ||
-			ProxyCompiler.bUsesObjectPosition ||
-			ProxyCompiler.bUsesObjectOrientation ||
-			ProxyCompiler.bUsesObjectRadius ||
-			ProxyCompiler.bUsesObjectBounds ||
-			ProxyCompiler.bUsesObjectLocalBounds ||
-			ProxyCompiler.bUsesPreSkinnedLocalBounds ||
-			ProxyCompiler.bUsesCustomPrimitiveData;
+		bRequiresPrimitiveData = ProxyCompiler.UsesPrimitiveData();
 	}, OutAnalysis);
 
 	// Also make sure the analysis takes into account primitive data
diff --git a/Engine/Plugins/Enterprise/GLTFExporter/Source/GLTFExporter/Private/Converters/GLTFProxyMaterialCompiler.h b/Engine/Plugins/Enterprise/GLTFExporter/Source/GLTFExporter/Private/Converters/GLTFProxyMaterialCompiler.h
--- a/Engine/Plugins/Enterprise/GLTFExporter/Source/GLTFExporter/Private/Converters/GLTFProxyMaterialCompiler.h
+++ b/Engine/Plugins/Enterprise/GLTFExporter/Source/GLTFExporter/Private/Converters/GLTFProxyMaterialCompiler.h
@@ -19,6 +19,22 @@ public:
 	bool bUsesObjectLocalBounds = false;
 	bool bUsesPreSkinnedLocalBounds = false;
 	bool bUsesCustomPrimitiveData = false;
+	bool bUsesTransformVector = false;
+	bool bUsesTransformPosition = false;
+
+	bool UsesPrimitiveData() const
+	{
+		return bUsesActorPosition
+			|| bUsesObjectPosition
+			|| bUsesObjectOrientation
+			|| bUsesObjectRadius
+			|| bUsesObjectBounds
+			|| bUsesObjectLocalBounds
+			|| bUsesPreSkinnedLocalBounds
+			|| bUsesCustomPrimitiveData
+			|| bUsesTransformVector
+			|| bUsesTransformPosition;
+	}
 
 	using FProxyMaterialCompiler::FProxyMaterialCompiler;
 
@@ -81,6 +97,18 @@ public:
 		return Compiler->CustomPrimitiveData(OutputIndex, Type);
 	}
 
+	virtual int32 TransformVector(EMaterialCommonBasis SourceCoordBasis, EMaterialCommonBasis DestCoordBasis, int32 A) override
+	{
+		if (SourceCoordBasis != DestCoordBasis) bUsesTransformVector = true; // TODO: optimize later by ignoring some irrelevant basis (like MCB_View and MCB_Camera)
+		return Compiler->TransformVector(SourceCoordBasis, DestCoordBasis, A);
+	}
+
+	virtual int32 TransformPosition(EMaterialCommonBasis SourceCoordBasis, EMaterialCommonBasis DestCoordBasis, int32 A) override
+	{
+		if (SourceCoordBasis != DestCoordBasis) bUsesTransformPosition = true; // TODO: optimize later by ignoring some irrelevant basis (like MCB_View and MCB_Camera)
+		return Compiler->TransformPosition(SourceCoordBasis, DestCoordBasis, A);
+	}
+
 	virtual int32 GIReplace(int32 Direct, int32 StaticIndirect, int32 DynamicIndirect) override
 	{
 		return Direct; // ignore all non-default branches in this runtime switch
