Fix binormal vector sign in the vertex tangent data not being correctly export to glTF

diff --git a/Engine/Plugins/Enterprise/GLTFExporter/Source/GLTFExporter/Private/Converters/GLTFAccessorConverters.cpp b/Engine/Plugins/Enterprise/GLTFExporter/Source/GLTFExporter/Private/Converters/GLTFAccessorConverters.cpp
--- a/Engine/Plugins/Enterprise/GLTFExporter/Source/GLTFExporter/Private/Converters/GLTFAccessorConverters.cpp
+++ b/Engine/Plugins/Enterprise/GLTFExporter/Source/GLTFExporter/Private/Converters/GLTFAccessorConverters.cpp
@@ -32,18 +32,18 @@ namespace
 	template <typename DestinationType, typename SourceType>
 	struct TGLTFVertexTangentUtilities
 	{
-		static DestinationType Convert(const SourceType& Tangent)
+		static DestinationType Convert(const SourceType& TangentX, const SourceType& TangentZ)
 		{
-			return FGLTFCoreUtilities::ConvertTangent(Tangent);
+			return FGLTFCoreUtilities::ConvertTangent(TangentX, TangentZ);
 		}
 	};
 
 	template <typename SourceType>
 	struct TGLTFVertexTangentUtilities<FGLTFVector4, SourceType>
 	{
-		static FGLTFVector4 Convert(const SourceType& Tangent)
+		static FGLTFVector4 Convert(const SourceType& TangentX, const SourceType& TangentZ)
 		{
-			return FGLTFCoreUtilities::ConvertTangent(Tangent.ToFVector4f());
+			return FGLTFCoreUtilities::ConvertTangent(TangentX.ToFVector3f(), TangentZ.ToFVector4f());
 		}
 	};
 }
@@ -212,8 +212,8 @@ FGLTFJsonBufferView* FGLTFNormalBufferConverter::ConvertBufferView(const FGLTFMe
 	const TArray<uint32>& IndexMap = MeshSection->IndexMap;
 	const uint32 VertexCount = IndexMap.Num();
 
-	typedef TStaticMeshVertexTangentDatum<SourceType> VertexTangentType;
-	const VertexTangentType* TangentData= static_cast<const VertexTangentType*>(SourceData);
+	typedef TStaticMeshVertexTangentDatum<SourceType> TangentDatumType;
+	const TangentDatumType* TangentData = static_cast<const TangentDatumType*>(SourceData);
 
 	TArray<DestinationType> Normals;
 	Normals.AddUninitialized(VertexCount);
@@ -221,8 +221,8 @@ FGLTFJsonBufferView* FGLTFNormalBufferConverter::ConvertBufferView(const FGLTFMe
 	for (uint32 VertexIndex = 0; VertexIndex < VertexCount; ++VertexIndex)
 	{
 		const uint32 MappedVertexIndex = IndexMap[VertexIndex];
-		const SourceType& Normal = TangentData[MappedVertexIndex].TangentZ;
-		Normals[VertexIndex] = TGLTFVertexNormalUtilities<DestinationType, SourceType>::Convert(Normal);
+		const TangentDatumType& TangentDatum = TangentData[MappedVertexIndex];
+		Normals[VertexIndex] = TGLTFVertexNormalUtilities<DestinationType, SourceType>::Convert(TangentDatum.TangentZ);
 	}
 
 	return Builder.AddBufferView(Normals, EGLTFJsonBufferTarget::ArrayBuffer);
@@ -289,8 +289,8 @@ FGLTFJsonBufferView* FGLTFTangentBufferConverter::ConvertBufferView(const FGLTFM
 	const TArray<uint32>& IndexMap = MeshSection->IndexMap;
 	const uint32 VertexCount = IndexMap.Num();
 
-	typedef TStaticMeshVertexTangentDatum<SourceType> VertexTangentType;
-	const VertexTangentType* VertexTangents = static_cast<const VertexTangentType*>(SourceData);
+	typedef TStaticMeshVertexTangentDatum<SourceType> TangentDatumType;
+	const TangentDatumType* VertexData = static_cast<const TangentDatumType*>(SourceData);
 
 	TArray<DestinationType> Tangents;
 	Tangents.AddUninitialized(VertexCount);
@@ -298,8 +298,8 @@ FGLTFJsonBufferView* FGLTFTangentBufferConverter::ConvertBufferView(const FGLTFM
 	for (uint32 VertexIndex = 0; VertexIndex < VertexCount; ++VertexIndex)
 	{
 		const uint32 MappedVertexIndex = IndexMap[VertexIndex];
-		const SourceType& Tangent = VertexTangents[MappedVertexIndex].TangentX;
-		Tangents[VertexIndex] = TGLTFVertexTangentUtilities<DestinationType, SourceType>::Convert(Tangent);
+		const TangentDatumType& TangentDatum = VertexData[MappedVertexIndex];
+		Tangents[VertexIndex] = TGLTFVertexTangentUtilities<DestinationType, SourceType>::Convert(TangentDatum.TangentX, TangentDatum.TangentZ);
 	}
 
 	return Builder.AddBufferView(Tangents, EGLTFJsonBufferTarget::ArrayBuffer);
diff --git a/Engine/Plugins/Enterprise/GLTFExporter/Source/GLTFExporter/Private/Utilities/GLTFCoreUtilities.cpp b/Engine/Plugins/Enterprise/GLTFExporter/Source/GLTFExporter/Private/Utilities/GLTFCoreUtilities.cpp
--- a/Engine/Plugins/Enterprise/GLTFExporter/Source/GLTFExporter/Private/Utilities/GLTFCoreUtilities.cpp
+++ b/Engine/Plugins/Enterprise/GLTFExporter/Source/GLTFExporter/Private/Utilities/GLTFCoreUtilities.cpp
@@ -43,22 +43,24 @@ FGLTFInt16Vector4 FGLTFCoreUtilities::ConvertNormal(const FPackedRGBA16N& Normal
 	return { SafeNormal.X, SafeNormal.Z, SafeNormal.Y, 0 };
 }
 
-FGLTFVector4 FGLTFCoreUtilities::ConvertTangent(const FVector4f& Tangent)
+FGLTFVector4 FGLTFCoreUtilities::ConvertTangent(const FVector3f& Tangent, const FVector4f& Normal)
 {
+	// Unreal keeps the binormal sign in the normal's w-component.
+	// glTF keeps the binormal sign in the tangent's w-component.
 	const FVector3f SafeTangent = Tangent.GetSafeNormal();
-	return { SafeTangent.X, SafeTangent.Z, SafeTangent.Y, Tangent.W < 0 ? -1.0f : +1.0f };
+	return { SafeTangent.X, SafeTangent.Z, SafeTangent.Y, Normal.W >= 0 ? 1.0f : -1.0f };
 }
 
-FGLTFInt8Vector4 FGLTFCoreUtilities::ConvertTangent(const FPackedNormal& Tangent)
+FGLTFInt8Vector4 FGLTFCoreUtilities::ConvertTangent(const FPackedNormal& Tangent, const FPackedNormal& Normal)
 {
 	const FPackedNormal SafeTangent = Tangent.ToFVector3f().GetSafeNormal();
-	return { SafeTangent.Vector.X, SafeTangent.Vector.Z, SafeTangent.Vector.Y, Tangent.Vector.W < 0 ? MIN_int8 : MAX_int8 };
+	return { SafeTangent.Vector.X, SafeTangent.Vector.Z, SafeTangent.Vector.Y, Normal.Vector.W >= 0 ? MAX_int8 : MIN_int8 };
 }
 
-FGLTFInt16Vector4 FGLTFCoreUtilities::ConvertTangent(const FPackedRGBA16N& Tangent)
+FGLTFInt16Vector4 FGLTFCoreUtilities::ConvertTangent(const FPackedRGBA16N& Tangent, const FPackedRGBA16N& Normal)
 {
 	const FPackedRGBA16N SafeTangent = Tangent.ToFVector3f().GetSafeNormal();
-	return { SafeTangent.X, SafeTangent.Z, SafeTangent.Y, Tangent.W < 0 ? MIN_int16 : MAX_int16 };
+	return { SafeTangent.X, SafeTangent.Z, SafeTangent.Y, Normal.W >= 0 ? MAX_int16 : MIN_int16 };
 }
 
 FGLTFVector2 FGLTFCoreUtilities::ConvertUV(const FVector2f& UV)
diff --git a/Engine/Plugins/Enterprise/GLTFExporter/Source/GLTFExporter/Public/Utilities/GLTFCoreUtilities.h b/Engine/Plugins/Enterprise/GLTFExporter/Source/GLTFExporter/Public/Utilities/GLTFCoreUtilities.h
--- a/Engine/Plugins/Enterprise/GLTFExporter/Source/GLTFExporter/Public/Utilities/GLTFCoreUtilities.h
+++ b/Engine/Plugins/Enterprise/GLTFExporter/Source/GLTFExporter/Public/Utilities/GLTFCoreUtilities.h
@@ -29,9 +29,9 @@ struct GLTFEXPORTER_API FGLTFCoreUtilities
 	static FGLTFInt8Vector4 ConvertNormal(const FPackedNormal& Normal);
 	static FGLTFInt16Vector4 ConvertNormal(const FPackedRGBA16N& Normal);
 
-	static FGLTFVector4 ConvertTangent(const FVector4f& Tangent);
-	static FGLTFInt8Vector4 ConvertTangent(const FPackedNormal& Tangent);
-	static FGLTFInt16Vector4 ConvertTangent(const FPackedRGBA16N& Tangent);
+	static FGLTFVector4 ConvertTangent(const FVector3f& Tangent, const FVector4f& Normal = FVector4f(ForceInitToZero));
+	static FGLTFInt8Vector4 ConvertTangent(const FPackedNormal& Tangent, const FPackedNormal& Normal = FPackedNormal());
+	static FGLTFInt16Vector4 ConvertTangent(const FPackedRGBA16N& Tangent, const FPackedRGBA16N& Normal = FPackedRGBA16N());
 
 	static FGLTFVector2 ConvertUV(const FVector2f& UV);
 	static FGLTFVector2 ConvertUV(const FVector2DHalf& UV);
