Use the same reference skeleton for skeletal animation and skeletal mesh

diff --git a/Engine/Plugins/Enterprise/GLTFExporter/Source/GLTFExporter/Private/Converters/GLTFBoneUtilities.cpp b/Engine/Plugins/Enterprise/GLTFExporter/Source/GLTFExporter/Private/Converters/GLTFBoneUtilities.cpp
--- a/Engine/Plugins/Enterprise/GLTFExporter/Source/GLTFExporter/Private/Converters/GLTFBoneUtilities.cpp
+++ b/Engine/Plugins/Enterprise/GLTFExporter/Source/GLTFExporter/Private/Converters/GLTFBoneUtilities.cpp
@@ -1,17 +1,23 @@
 // Copyright Epic Games, Inc. All Rights Reserved.
 
 #include "Converters/GLTFBoneUtilities.h"
-#include "BonePose.h"
-#include "Sections/MovieScene3DTransformSection.h"
 #include "Animation/AnimationPoseData.h"
 #include "Animation/AnimSequence.h"
 #include "Animation/AttributesRuntime.h"
+#include "Engine/SkeletalMesh.h"
 #include "ReferenceSkeleton.h"
+#include "BonePose.h"
+
+const FReferenceSkeleton& FGLTFBoneUtilities::GetReferenceSkeleton(const USkeletalMesh* SkeletalMesh)
+{
+	const USkeleton* Skeleton = SkeletalMesh->GetSkeleton();
+	return Skeleton != nullptr ? Skeleton->GetReferenceSkeleton() : SkeletalMesh->GetRefSkeleton();
+}
 
-FTransform FGLTFBoneUtilities::GetBindTransform(const FReferenceSkeleton& RefSkeleton, int32 BoneIndex)
+FTransform FGLTFBoneUtilities::GetBindTransform(const FReferenceSkeleton& ReferenceSkeleton, int32 BoneIndex)
 {
-	const TArray<FMeshBoneInfo>& BoneInfos = RefSkeleton.GetRefBoneInfo();
-	const TArray<FTransform>& BonePoses = RefSkeleton.GetRefBonePose();
+	const TArray<FMeshBoneInfo>& BoneInfos = ReferenceSkeleton.GetRefBoneInfo();
+	const TArray<FTransform>& BonePoses = ReferenceSkeleton.GetRefBonePose();
 
 	int32 CurBoneIndex = BoneIndex;
 	FTransform BindTransform = FTransform::Identity;
diff --git a/Engine/Plugins/Enterprise/GLTFExporter/Source/GLTFExporter/Private/Converters/GLTFBoneUtilities.h b/Engine/Plugins/Enterprise/GLTFExporter/Source/GLTFExporter/Private/Converters/GLTFBoneUtilities.h
--- a/Engine/Plugins/Enterprise/GLTFExporter/Source/GLTFExporter/Private/Converters/GLTFBoneUtilities.h
+++ b/Engine/Plugins/Enterprise/GLTFExporter/Source/GLTFExporter/Private/Converters/GLTFBoneUtilities.h
@@ -5,15 +5,16 @@
 #include "CoreMinimal.h"
 #include "BoneIndices.h"
 
-struct FMovieSceneDoubleChannel;
-struct FReferenceSkeleton;
 class UAnimSequence;
-class UMovieScene3DTransformSection;
+class USkeletalMesh;
 class USkeleton;
+struct FReferenceSkeleton;
 
 struct FGLTFBoneUtilities
 {
-	static FTransform GetBindTransform(const FReferenceSkeleton& RefSkeleton, int32 BoneIndex);
+	static const FReferenceSkeleton& GetReferenceSkeleton(const USkeletalMesh* SkeletalMesh);
+
+	static FTransform GetBindTransform(const FReferenceSkeleton& ReferenceSkeleton, int32 BoneIndex);
 
 	static void GetFrameTimestamps(const UAnimSequence* AnimSequence, TArray<float>& OutFrameTimestamps);
 
diff --git a/Engine/Plugins/Enterprise/GLTFExporter/Source/GLTFExporter/Private/Converters/GLTFNodeConverters.cpp b/Engine/Plugins/Enterprise/GLTFExporter/Source/GLTFExporter/Private/Converters/GLTFNodeConverters.cpp
--- a/Engine/Plugins/Enterprise/GLTFExporter/Source/GLTFExporter/Private/Converters/GLTFNodeConverters.cpp
+++ b/Engine/Plugins/Enterprise/GLTFExporter/Source/GLTFExporter/Private/Converters/GLTFNodeConverters.cpp
@@ -1,9 +1,10 @@
 // Copyright Epic Games, Inc. All Rights Reserved.
 
 #include "Converters/GLTFNodeConverters.h"
+#include "Converters/GLTFNameUtilities.h"
+#include "Converters/GLTFBoneUtilities.h"
 #include "Builders/GLTFContainerBuilder.h"
 #include "Utilities/GLTFCoreUtilities.h"
-#include "Converters/GLTFNameUtilities.h"
 #include "LevelSequenceActor.h"
 #include "Camera/CameraComponent.h"
 #include "Components/LightComponent.h"
@@ -229,7 +230,7 @@ FGLTFJsonNode* FGLTFStaticSocketConverter::Convert(FGLTFJsonNode* RootNode, cons
 
 FGLTFJsonNode* FGLTFSkeletalSocketConverter::Convert(FGLTFJsonNode* RootNode, const USkeletalMesh* SkeletalMesh, FName SocketName)
 {
-	const FReferenceSkeleton& RefSkeleton = SkeletalMesh->GetRefSkeleton();
+	const FReferenceSkeleton& ReferenceSkeleton = FGLTFBoneUtilities::GetReferenceSkeleton(SkeletalMesh);
 
 	const USkeletalMeshSocket* Socket = SkeletalMesh->FindSocket(SocketName);
 	if (Socket != nullptr)
@@ -242,14 +243,14 @@ FGLTFJsonNode* FGLTFSkeletalSocketConverter::Convert(FGLTFJsonNode* RootNode, co
 		Node->Rotation = FGLTFCoreUtilities::ConvertRotation(FRotator3f(Socket->RelativeRotation).Quaternion());
 		Node->Scale = FGLTFCoreUtilities::ConvertScale(FVector3f(Socket->RelativeScale));
 
-		const int32 ParentBone = RefSkeleton.FindBoneIndex(Socket->BoneName);
+		const int32 ParentBone = ReferenceSkeleton.FindBoneIndex(Socket->BoneName);
 		FGLTFJsonNode* ParentNode = ParentBone != INDEX_NONE ? Builder.AddUniqueNode(RootNode, SkeletalMesh, ParentBone) : RootNode;
 
 		ParentNode->Children.Add(Node);
 		return Node;
 	}
 
-	const int32 BoneIndex = RefSkeleton.FindBoneIndex(SocketName);
+	const int32 BoneIndex = ReferenceSkeleton.FindBoneIndex(SocketName);
 	if (BoneIndex != INDEX_NONE)
 	{
 		return Builder.AddUniqueNode(RootNode, SkeletalMesh, BoneIndex);
@@ -261,11 +262,11 @@ FGLTFJsonNode* FGLTFSkeletalSocketConverter::Convert(FGLTFJsonNode* RootNode, co
 
 FGLTFJsonNode* FGLTFSkeletalBoneConverter::Convert(FGLTFJsonNode* RootNode, const USkeletalMesh* SkeletalMesh, int32 BoneIndex)
 {
-	const FReferenceSkeleton& RefSkeleton = SkeletalMesh->GetRefSkeleton();
+	const FReferenceSkeleton& ReferenceSkeleton = FGLTFBoneUtilities::GetReferenceSkeleton(SkeletalMesh);
 
 	// TODO: add support for [Principal]PoseComponent?
 
-	const TArray<FMeshBoneInfo>& BoneInfos = RefSkeleton.GetRefBoneInfo();
+	const TArray<FMeshBoneInfo>& BoneInfos = ReferenceSkeleton.GetRefBoneInfo();
 	if (!BoneInfos.IsValidIndex(BoneIndex))
 	{
 		// TODO: report error
@@ -277,7 +278,7 @@ FGLTFJsonNode* FGLTFSkeletalBoneConverter::Convert(FGLTFJsonNode* RootNode, cons
 	FGLTFJsonNode* Node = Builder.AddNode();
 	Node->Name = BoneInfo.Name.ToString();
 
-	const TArray<FTransform>& BonePoses = RefSkeleton.GetRefBonePose();
+	const TArray<FTransform>& BonePoses = ReferenceSkeleton.GetRefBonePose();
 	if (BonePoses.IsValidIndex(BoneIndex))
 	{
 		// TODO: add warning check for non-uniform scaling
diff --git a/Engine/Plugins/Enterprise/GLTFExporter/Source/GLTFExporter/Private/Converters/GLTFSkinConverters.cpp b/Engine/Plugins/Enterprise/GLTFExporter/Source/GLTFExporter/Private/Converters/GLTFSkinConverters.cpp
--- a/Engine/Plugins/Enterprise/GLTFExporter/Source/GLTFExporter/Private/Converters/GLTFSkinConverters.cpp
+++ b/Engine/Plugins/Enterprise/GLTFExporter/Source/GLTFExporter/Private/Converters/GLTFSkinConverters.cpp
@@ -10,9 +10,9 @@
 FGLTFJsonSkin* FGLTFSkinConverter::Convert(FGLTFJsonNode* RootNode, const USkeletalMesh* SkeletalMesh)
 {
 	const USkeleton* Skeleton = SkeletalMesh->GetSkeleton();
-	const FReferenceSkeleton& RefSkeleton = SkeletalMesh->GetRefSkeleton();
+	const FReferenceSkeleton& ReferenceSkeleton = FGLTFBoneUtilities::GetReferenceSkeleton(SkeletalMesh);
 
-	const int32 BoneCount = RefSkeleton.GetNum();
+	const int32 BoneCount = ReferenceSkeleton.GetNum();
 	if (BoneCount == 0)
 	{
 		// TODO: report warning
@@ -34,7 +34,7 @@ FGLTFJsonSkin* FGLTFSkinConverter::Convert(FGLTFJsonNode* RootNode, const USkele
 
 	for (int32 BoneIndex = 0; BoneIndex < BoneCount; ++BoneIndex)
 	{
-		const FTransform3f InverseBindTransform = FTransform3f(FGLTFBoneUtilities::GetBindTransform(RefSkeleton, BoneIndex).Inverse());
+		const FTransform3f InverseBindTransform = FTransform3f(FGLTFBoneUtilities::GetBindTransform(ReferenceSkeleton, BoneIndex).Inverse());
 		InverseBindMatrices[BoneIndex] = FGLTFCoreUtilities::ConvertTransform(InverseBindTransform, Builder.ExportOptions->ExportUniformScale);
 	}
 
