Replace illegal chars with '_' (instead of just removing them) when sanitizing filenames in GLTFExporter

diff --git a/Engine/Plugins/Enterprise/GLTFExporter/Source/GLTFExporter/Private/Builders/GLTFFileBuilder.cpp b/Engine/Plugins/Enterprise/GLTFExporter/Source/GLTFExporter/Private/Builders/GLTFFileBuilder.cpp
--- a/Engine/Plugins/Enterprise/GLTFExporter/Source/GLTFExporter/Private/Builders/GLTFFileBuilder.cpp
+++ b/Engine/Plugins/Enterprise/GLTFExporter/Source/GLTFExporter/Private/Builders/GLTFFileBuilder.cpp
@@ -70,17 +70,23 @@ FString FGLTFFileBuilder::GetUniqueFileName(const FString& InFileName) const
 
 FString FGLTFFileBuilder::SanitizeFileName(const FString& InFileName)
 {
-	static const TCHAR* IllegalChars = TEXT("<>:\"/\\|?*");
+	static constexpr TCHAR IllegalChars[] = {
+		TEXT('<'),
+		TEXT('>'),
+		TEXT(':'),
+		TEXT('"'),
+		TEXT('/'),
+		TEXT('\\'),
+		TEXT('|'),
+		TEXT('?'),
+		TEXT('*'),
+	};
 
-	FString SanitizedFileName;
-	SanitizedFileName.Reserve(InFileName.Len());
+	FString SanitizedFileName = InFileName;
 
-	for (const TCHAR& Chr : InFileName)
+	for (const TCHAR& IllegalChar : IllegalChars)
 	{
-		if (FCString::Strchr(IllegalChars, Chr) == nullptr)
-		{
-			SanitizedFileName += Chr;
-		}
+		SanitizedFileName.ReplaceCharInline(IllegalChar, TEXT('_'));
 	}
 
 	return SanitizedFileName;
