Fix incorrect reference skeleton used when exporting skeletal animations to glTF

diff --git a/Engine/Plugins/Enterprise/GLTFExporter/Source/GLTFExporter/Private/Converters/GLTFBoneUtilities.cpp b/Engine/Plugins/Enterprise/GLTFExporter/Source/GLTFExporter/Private/Converters/GLTFBoneUtilities.cpp
--- a/Engine/Plugins/Enterprise/GLTFExporter/Source/GLTFExporter/Private/Converters/GLTFBoneUtilities.cpp
+++ b/Engine/Plugins/Enterprise/GLTFExporter/Source/GLTFExporter/Private/Converters/GLTFBoneUtilities.cpp
@@ -4,33 +4,37 @@
 #include "Animation/AnimationPoseData.h"
 #include "Animation/AnimSequence.h"
 #include "Animation/AttributesRuntime.h"
-#include "Engine/SkeletalMesh.h"
 #include "ReferenceSkeleton.h"
 #include "BonePose.h"
 
-const FReferenceSkeleton& FGLTFBoneUtilities::GetReferenceSkeleton(const USkeletalMesh* SkeletalMesh)
-{
-	const USkeleton* Skeleton = SkeletalMesh->GetSkeleton();
-	return Skeleton != nullptr ? Skeleton->GetReferenceSkeleton() : SkeletalMesh->GetRefSkeleton();
-}
-
 FTransform FGLTFBoneUtilities::GetBindTransform(const FReferenceSkeleton& ReferenceSkeleton, int32 BoneIndex)
 {
 	const TArray<FMeshBoneInfo>& BoneInfos = ReferenceSkeleton.GetRefBoneInfo();
 	const TArray<FTransform>& BonePoses = ReferenceSkeleton.GetRefBonePose();
 
-	int32 CurBoneIndex = BoneIndex;
-	FTransform BindTransform = FTransform::Identity;
+	FTransform BindTransform = BonePoses[BoneIndex];
+	int32 ParentBoneIndex = BoneInfos[BoneIndex].ParentIndex;
 
-	do
+	while (ParentBoneIndex != INDEX_NONE)
 	{
-		BindTransform = BindTransform * BonePoses[CurBoneIndex];
-		CurBoneIndex = BoneInfos[CurBoneIndex].ParentIndex;
-	} while (CurBoneIndex != INDEX_NONE);
+		BindTransform = BindTransform * BonePoses[ParentBoneIndex];
+		ParentBoneIndex = BoneInfos[ParentBoneIndex].ParentIndex;
+	}
 
 	return BindTransform;
 }
 
+void FGLTFBoneUtilities::GetBoneIndices(const FReferenceSkeleton& ReferenceSkeleton, TArray<FBoneIndexType>& OutBoneIndices)
+{
+	const int32 BoneCount = ReferenceSkeleton.GetNum();
+	OutBoneIndices.AddUninitialized(BoneCount);
+
+	for (int32 BoneIndex = 0; BoneIndex < BoneCount; ++BoneIndex)
+	{
+		OutBoneIndices[BoneIndex] = static_cast<FBoneIndexType>(BoneIndex);
+	}
+}
+
 void FGLTFBoneUtilities::GetFrameTimestamps(const UAnimSequence* AnimSequence, TArray<float>& OutFrameTimestamps)
 {
 	const int32 FrameCount = AnimSequence->GetNumberOfSampledKeys();
@@ -45,25 +49,13 @@ void FGLTFBoneUtilities::GetFrameTimestamps(const UAnimSequence* AnimSequence, T
 	}
 }
 
-void FGLTFBoneUtilities::GetBoneIndices(const USkeleton* Skeleton, TArray<FBoneIndexType>& OutBoneIndices)
+void FGLTFBoneUtilities::GetBoneTransformsByFrame(const UObject* SkeletalMeshOrSkeleton, const UAnimSequence* AnimSequence, const TArray<FBoneIndexType>& BoneIndices, const TArray<float>& FrameTimestamps, TArray<TArray<FTransform>>& OutBoneTransformsByFrame)
 {
-	const int32 BoneCount = Skeleton->GetReferenceSkeleton().GetNum();
-	OutBoneIndices.AddUninitialized(BoneCount);
-
-	for (int32 BoneIndex = 0; BoneIndex < BoneCount; ++BoneIndex)
-	{
-		OutBoneIndices[BoneIndex] = static_cast<FBoneIndexType>(BoneIndex);
-	}
-}
-
-void FGLTFBoneUtilities::GetBoneTransformsByFrame(const UAnimSequence* AnimSequence, const TArray<float>& FrameTimestamps, const TArray<FBoneIndexType>& BoneIndices, TArray<TArray<FTransform>>& OutBoneTransformsByFrame)
-{
-
 	// Make sure to free stack allocations made by FCompactPose, FBlendedCurve, and FStackCustomAttributes when end of scope
 	FMemMark Mark(FMemStack::Get());
 
 	FBoneContainer BoneContainer;
-	BoneContainer.InitializeTo(BoneIndices, UE::Anim::FCurveFilterSettings(UE::Anim::ECurveFilterMode::DisallowAll), *AnimSequence->GetSkeleton());
+	BoneContainer.InitializeTo(BoneIndices, UE::Anim::FCurveFilterSettings(UE::Anim::ECurveFilterMode::DisallowAll), *const_cast<UObject*>(SkeletalMeshOrSkeleton));
 
 	const int32 FrameCount = FrameTimestamps.Num();
 	OutBoneTransformsByFrame.AddDefaulted(FrameCount);
diff --git a/Engine/Plugins/Enterprise/GLTFExporter/Source/GLTFExporter/Private/Converters/GLTFBoneUtilities.h b/Engine/Plugins/Enterprise/GLTFExporter/Source/GLTFExporter/Private/Converters/GLTFBoneUtilities.h
--- a/Engine/Plugins/Enterprise/GLTFExporter/Source/GLTFExporter/Private/Converters/GLTFBoneUtilities.h
+++ b/Engine/Plugins/Enterprise/GLTFExporter/Source/GLTFExporter/Private/Converters/GLTFBoneUtilities.h
@@ -6,19 +6,15 @@
 #include "BoneIndices.h"
 
 class UAnimSequence;
-class USkeletalMesh;
-class USkeleton;
 struct FReferenceSkeleton;
 
 struct FGLTFBoneUtilities
 {
-	static const FReferenceSkeleton& GetReferenceSkeleton(const USkeletalMesh* SkeletalMesh);
-
 	static FTransform GetBindTransform(const FReferenceSkeleton& ReferenceSkeleton, int32 BoneIndex);
 
-	static void GetFrameTimestamps(const UAnimSequence* AnimSequence, TArray<float>& OutFrameTimestamps);
+	static void GetBoneIndices(const FReferenceSkeleton& ReferenceSkeleton, TArray<FBoneIndexType>& OutBoneIndices);
 
-	static void GetBoneIndices(const USkeleton* Skeleton, TArray<FBoneIndexType>& OutBoneIndices);
+	static void GetFrameTimestamps(const UAnimSequence* AnimSequence, TArray<float>& OutFrameTimestamps);
 
-	static void GetBoneTransformsByFrame(const UAnimSequence* AnimSequence, const TArray<float>& FrameTimestamps, const TArray<FBoneIndexType>& BoneIndices, TArray<TArray<FTransform>>& OutBoneTransformsByFrame);
+	static void GetBoneTransformsByFrame(const UObject* SkeletalMeshOrSkeleton, const UAnimSequence* AnimSequence, const TArray<FBoneIndexType>& BoneIndices, const TArray<float>& FrameTimestamps, TArray<TArray<FTransform>>& OutBoneTransformsByFrame);
 };
diff --git a/Engine/Plugins/Enterprise/GLTFExporter/Source/GLTFExporter/Private/Converters/GLTFNodeConverters.cpp b/Engine/Plugins/Enterprise/GLTFExporter/Source/GLTFExporter/Private/Converters/GLTFNodeConverters.cpp
--- a/Engine/Plugins/Enterprise/GLTFExporter/Source/GLTFExporter/Private/Converters/GLTFNodeConverters.cpp
+++ b/Engine/Plugins/Enterprise/GLTFExporter/Source/GLTFExporter/Private/Converters/GLTFNodeConverters.cpp
@@ -230,7 +230,7 @@ FGLTFJsonNode* FGLTFStaticSocketConverter::Convert(FGLTFJsonNode* RootNode, cons
 
 FGLTFJsonNode* FGLTFSkeletalSocketConverter::Convert(FGLTFJsonNode* RootNode, const USkeletalMesh* SkeletalMesh, FName SocketName)
 {
-	const FReferenceSkeleton& ReferenceSkeleton = FGLTFBoneUtilities::GetReferenceSkeleton(SkeletalMesh);
+	const FReferenceSkeleton& ReferenceSkeleton = SkeletalMesh->GetRefSkeleton();
 
 	const USkeletalMeshSocket* Socket = SkeletalMesh->FindSocket(SocketName);
 	if (Socket != nullptr)
@@ -262,9 +262,9 @@ FGLTFJsonNode* FGLTFSkeletalSocketConverter::Convert(FGLTFJsonNode* RootNode, co
 
 FGLTFJsonNode* FGLTFSkeletalBoneConverter::Convert(FGLTFJsonNode* RootNode, const USkeletalMesh* SkeletalMesh, int32 BoneIndex)
 {
-	const FReferenceSkeleton& ReferenceSkeleton = FGLTFBoneUtilities::GetReferenceSkeleton(SkeletalMesh);
+	const FReferenceSkeleton& ReferenceSkeleton = SkeletalMesh->GetRefSkeleton();
 
-	// TODO: add support for [Principal]PoseComponent?
+	// TODO: add support for LeaderPoseComponent?
 
 	const TArray<FMeshBoneInfo>& BoneInfos = ReferenceSkeleton.GetRefBoneInfo();
 	if (!BoneInfos.IsValidIndex(BoneIndex))
diff --git a/Engine/Plugins/Enterprise/GLTFExporter/Source/GLTFExporter/Private/Converters/GLTFSkinConverters.cpp b/Engine/Plugins/Enterprise/GLTFExporter/Source/GLTFExporter/Private/Converters/GLTFSkinConverters.cpp
--- a/Engine/Plugins/Enterprise/GLTFExporter/Source/GLTFExporter/Private/Converters/GLTFSkinConverters.cpp
+++ b/Engine/Plugins/Enterprise/GLTFExporter/Source/GLTFExporter/Private/Converters/GLTFSkinConverters.cpp
@@ -4,13 +4,11 @@
 #include "Utilities/GLTFCoreUtilities.h"
 #include "Converters/GLTFBoneUtilities.h"
 #include "Builders/GLTFConvertBuilder.h"
-#include "Animation/Skeleton.h"
 #include "Engine/SkeletalMesh.h"
 
 FGLTFJsonSkin* FGLTFSkinConverter::Convert(FGLTFJsonNode* RootNode, const USkeletalMesh* SkeletalMesh)
 {
-	const USkeleton* Skeleton = SkeletalMesh->GetSkeleton();
-	const FReferenceSkeleton& ReferenceSkeleton = FGLTFBoneUtilities::GetReferenceSkeleton(SkeletalMesh);
+	const FReferenceSkeleton& ReferenceSkeleton = SkeletalMesh->GetRefSkeleton();
 
 	const int32 BoneCount = ReferenceSkeleton.GetNum();
 	if (BoneCount == 0)
@@ -20,7 +18,7 @@ FGLTFJsonSkin* FGLTFSkinConverter::Convert(FGLTFJsonNode* RootNode, const USkele
 	}
 
 	FGLTFJsonSkin* JsonSkin = Builder.AddSkin();
-	JsonSkin->Name = Skeleton != nullptr ? Skeleton->GetName() : SkeletalMesh->GetName();
+	JsonSkin->Name = SkeletalMesh->GetName();
 	JsonSkin->Skeleton = RootNode;
 	JsonSkin->Joints.AddUninitialized(BoneCount);
 
diff --git a/Engine/Plugins/Enterprise/GLTFExporter/Source/GLTFExporter/Private/Tasks/GLTFDelayedAnimationTasks.cpp b/Engine/Plugins/Enterprise/GLTFExporter/Source/GLTFExporter/Private/Tasks/GLTFDelayedAnimationTasks.cpp
--- a/Engine/Plugins/Enterprise/GLTFExporter/Source/GLTFExporter/Private/Tasks/GLTFDelayedAnimationTasks.cpp
+++ b/Engine/Plugins/Enterprise/GLTFExporter/Source/GLTFExporter/Private/Tasks/GLTFDelayedAnimationTasks.cpp
@@ -44,10 +44,10 @@ void FGLTFDelayedAnimSequenceTask::Process()
 	}
 
 	TArray<FBoneIndexType> BoneIndices;
-	FGLTFBoneUtilities::GetBoneIndices(AnimSequence->GetSkeleton(), BoneIndices);
+	FGLTFBoneUtilities::GetBoneIndices(SkeletalMesh->GetRefSkeleton(), BoneIndices);
 
 	TArray<TArray<FTransform>> FrameTransforms;
-	FGLTFBoneUtilities::GetBoneTransformsByFrame(AnimSequence, Timestamps, BoneIndices, FrameTransforms);
+	FGLTFBoneUtilities::GetBoneTransformsByFrame(SkeletalMesh, AnimSequence, BoneIndices, Timestamps, FrameTransforms);
 
 	for (const FBoneIndexType BoneIndex : BoneIndices)
 	{
