Fix horde issue 282011

diff --git a/Engine/Plugins/Enterprise/GLTFExporter/Source/GLTFExporter/Private/Tasks/GLTFDelayedAnimationTasks.cpp b/Engine/Plugins/Enterprise/GLTFExporter/Source/GLTFExporter/Private/Tasks/GLTFDelayedAnimationTasks.cpp
--- a/Engine/Plugins/Enterprise/GLTFExporter/Source/GLTFExporter/Private/Tasks/GLTFDelayedAnimationTasks.cpp
+++ b/Engine/Plugins/Enterprise/GLTFExporter/Source/GLTFExporter/Private/Tasks/GLTFDelayedAnimationTasks.cpp
@@ -2,9 +2,10 @@
 
 #include "Tasks/GLTFDelayedAnimationTasks.h"
 #include "Builders/GLTFContainerBuilder.h"
-#include "Engine/Level.h"
 #include "Utilities/GLTFCoreUtilities.h"
 #include "Converters/GLTFBoneUtilities.h"
+#include "Engine/SkeletalMesh.h"
+#include "Engine/Level.h"
 #include "LevelSequence.h"
 #include "LevelSequenceActor.h"
 #include "LevelSequencePlayer.h"
