Fix missing texture transform for occlusion and normal textures in exported glTF

diff --git a/Engine/Plugins/Enterprise/GLTFExporter/Source/GLTFExporter/Private/Json/GLTFJsonMaterial.cpp b/Engine/Plugins/Enterprise/GLTFExporter/Source/GLTFExporter/Private/Json/GLTFJsonMaterial.cpp
--- a/Engine/Plugins/Enterprise/GLTFExporter/Source/GLTFExporter/Private/Json/GLTFJsonMaterial.cpp
+++ b/Engine/Plugins/Enterprise/GLTFExporter/Source/GLTFExporter/Private/Json/GLTFJsonMaterial.cpp
@@ -33,6 +33,13 @@ void FGLTFJsonNormalTextureInfo::WriteObject(IGLTFJsonWriter& Writer) const
 	{
 		Writer.Write(TEXT("scale"), Scale);
 	}
+
+	if (!Transform.IsNearlyDefault(Writer.DefaultTolerance))
+	{
+		Writer.StartExtensions();
+		Writer.Write(EGLTFJsonExtension::KHR_TextureTransform, Transform);
+		Writer.EndExtensions();
+	}
 }
 
 void FGLTFJsonOcclusionTextureInfo::WriteObject(IGLTFJsonWriter& Writer) const
@@ -48,6 +55,13 @@ void FGLTFJsonOcclusionTextureInfo::WriteObject(IGLTFJsonWriter& Writer) const
 	{
 		Writer.Write(TEXT("strength"), Strength);
 	}
+
+	if (!Transform.IsNearlyDefault(Writer.DefaultTolerance))
+	{
+		Writer.StartExtensions();
+		Writer.Write(EGLTFJsonExtension::KHR_TextureTransform, Transform);
+		Writer.EndExtensions();
+	}
 }
 
 void FGLTFJsonPBRMetallicRoughness::WriteObject(IGLTFJsonWriter& Writer) const
