Add support for KHR_materials_emissive_strength in glTF exporter

diff --git a/Engine/Plugins/Enterprise/GLTFExporter/Source/GLTFExporter/Private/Converters/GLTFMaterialProxyFactory.cpp b/Engine/Plugins/Enterprise/GLTFExporter/Source/GLTFExporter/Private/Converters/GLTFMaterialProxyFactory.cpp
--- a/Engine/Plugins/Enterprise/GLTFExporter/Source/GLTFExporter/Private/Converters/GLTFMaterialProxyFactory.cpp
+++ b/Engine/Plugins/Enterprise/GLTFExporter/Source/GLTFExporter/Private/Converters/GLTFMaterialProxyFactory.cpp
@@ -99,34 +99,35 @@ void FGLTFMaterialProxyFactory::SetBaseProperties(UMaterialInstanceConstant* Pro
 
 void FGLTFMaterialProxyFactory::SetProxyParameters(UMaterialInstanceConstant* ProxyMaterial, const FGLTFJsonMaterial& JsonMaterial)
 {
-	SetProxyParameter(ProxyMaterial, FGLTFProxyMaterialInfo::BaseColorFactor, JsonMaterial.PBRMetallicRoughness.BaseColorFactor);
 	SetProxyParameter(ProxyMaterial, FGLTFProxyMaterialInfo::BaseColor, JsonMaterial.PBRMetallicRoughness.BaseColorTexture);
+	SetProxyParameter(ProxyMaterial, FGLTFProxyMaterialInfo::BaseColorFactor, JsonMaterial.PBRMetallicRoughness.BaseColorFactor);
 
 	if (JsonMaterial.ShadingModel == EGLTFJsonShadingModel::Default || JsonMaterial.ShadingModel == EGLTFJsonShadingModel::ClearCoat)
 	{
-		SetProxyParameter(ProxyMaterial, FGLTFProxyMaterialInfo::EmissiveFactor, JsonMaterial.EmissiveFactor);
 		SetProxyParameter(ProxyMaterial, FGLTFProxyMaterialInfo::Emissive, JsonMaterial.EmissiveTexture);
+		SetProxyParameter(ProxyMaterial, FGLTFProxyMaterialInfo::EmissiveFactor, JsonMaterial.EmissiveFactor);
+		SetProxyParameter(ProxyMaterial, FGLTFProxyMaterialInfo::EmissiveStrength, JsonMaterial.EmissiveStrength);
 
+		SetProxyParameter(ProxyMaterial, FGLTFProxyMaterialInfo::MetallicRoughness, JsonMaterial.PBRMetallicRoughness.MetallicRoughnessTexture);
 		SetProxyParameter(ProxyMaterial, FGLTFProxyMaterialInfo::MetallicFactor, JsonMaterial.PBRMetallicRoughness.MetallicFactor);
 		SetProxyParameter(ProxyMaterial, FGLTFProxyMaterialInfo::RoughnessFactor, JsonMaterial.PBRMetallicRoughness.RoughnessFactor);
-		SetProxyParameter(ProxyMaterial, FGLTFProxyMaterialInfo::MetallicRoughness, JsonMaterial.PBRMetallicRoughness.MetallicRoughnessTexture);
 
-		SetProxyParameter(ProxyMaterial, FGLTFProxyMaterialInfo::NormalScale, JsonMaterial.NormalTexture.Scale);
 		SetProxyParameter(ProxyMaterial, FGLTFProxyMaterialInfo::Normal, JsonMaterial.NormalTexture);
+		SetProxyParameter(ProxyMaterial, FGLTFProxyMaterialInfo::NormalScale, JsonMaterial.NormalTexture.Scale);
 
-		SetProxyParameter(ProxyMaterial, FGLTFProxyMaterialInfo::OcclusionStrength, JsonMaterial.OcclusionTexture.Strength);
 		SetProxyParameter(ProxyMaterial, FGLTFProxyMaterialInfo::Occlusion, JsonMaterial.OcclusionTexture);
+		SetProxyParameter(ProxyMaterial, FGLTFProxyMaterialInfo::OcclusionStrength, JsonMaterial.OcclusionTexture.Strength);
 
 		if (JsonMaterial.ShadingModel == EGLTFJsonShadingModel::ClearCoat)
 		{
-			SetProxyParameter(ProxyMaterial, FGLTFProxyMaterialInfo::ClearCoatFactor, JsonMaterial.ClearCoat.ClearCoatFactor);
 			SetProxyParameter(ProxyMaterial, FGLTFProxyMaterialInfo::ClearCoat, JsonMaterial.ClearCoat.ClearCoatTexture);
+			SetProxyParameter(ProxyMaterial, FGLTFProxyMaterialInfo::ClearCoatFactor, JsonMaterial.ClearCoat.ClearCoatFactor);
 
-			SetProxyParameter(ProxyMaterial, FGLTFProxyMaterialInfo::ClearCoatRoughnessFactor, JsonMaterial.ClearCoat.ClearCoatRoughnessFactor);
 			SetProxyParameter(ProxyMaterial, FGLTFProxyMaterialInfo::ClearCoatRoughness, JsonMaterial.ClearCoat.ClearCoatRoughnessTexture);
+			SetProxyParameter(ProxyMaterial, FGLTFProxyMaterialInfo::ClearCoatRoughnessFactor, JsonMaterial.ClearCoat.ClearCoatRoughnessFactor);
 
-			SetProxyParameter(ProxyMaterial, FGLTFProxyMaterialInfo::ClearCoatNormalScale, JsonMaterial.ClearCoat.ClearCoatNormalTexture.Scale);
 			SetProxyParameter(ProxyMaterial, FGLTFProxyMaterialInfo::ClearCoatNormal, JsonMaterial.ClearCoat.ClearCoatNormalTexture);
+			SetProxyParameter(ProxyMaterial, FGLTFProxyMaterialInfo::ClearCoatNormalScale, JsonMaterial.ClearCoat.ClearCoatNormalTexture.Scale);
 		}
 	}
 }
diff --git a/Engine/Plugins/Enterprise/GLTFExporter/Source/GLTFExporter/Private/Json/GLTFJsonMaterial.cpp b/Engine/Plugins/Enterprise/GLTFExporter/Source/GLTFExporter/Private/Json/GLTFJsonMaterial.cpp
--- a/Engine/Plugins/Enterprise/GLTFExporter/Source/GLTFExporter/Private/Json/GLTFJsonMaterial.cpp
+++ b/Engine/Plugins/Enterprise/GLTFExporter/Source/GLTFExporter/Private/Json/GLTFJsonMaterial.cpp
@@ -167,14 +167,15 @@ void FGLTFJsonMaterial::WriteObject(IGLTFJsonWriter& Writer) const
 		Writer.Write(TEXT("doubleSided"), DoubleSided);
 	}
 
-	if (ShadingModel == EGLTFJsonShadingModel::Unlit || ShadingModel == EGLTFJsonShadingModel::ClearCoat)
+	const bool HasEmissiveStrength = !FMath::IsNearlyEqual(EmissiveStrength, 1.0f, Writer.DefaultTolerance);
+	if (ShadingModel == EGLTFJsonShadingModel::Unlit || ShadingModel == EGLTFJsonShadingModel::ClearCoat || HasEmissiveStrength)
 	{
 		Writer.StartExtensions();
 
 		if (ShadingModel == EGLTFJsonShadingModel::Unlit)
 		{
-			Writer.StartExtension(EGLTFJsonExtension::KHR_MaterialsUnlit);
 			// Write empty object
+			Writer.StartExtension(EGLTFJsonExtension::KHR_MaterialsUnlit);
 			Writer.EndExtension();
 		}
 		else if (ShadingModel == EGLTFJsonShadingModel::ClearCoat)
@@ -182,6 +183,13 @@ void FGLTFJsonMaterial::WriteObject(IGLTFJsonWriter& Writer) const
 			Writer.Write(EGLTFJsonExtension::KHR_MaterialsClearCoat, ClearCoat);
 		}
 
+		if (HasEmissiveStrength)
+		{
+			Writer.StartExtension(EGLTFJsonExtension::KHR_MaterialsEmissiveStrength);
+			Writer.Write(TEXT("emissiveStrength"), EmissiveStrength);
+			Writer.EndExtension();
+		}
+
 		Writer.EndExtensions();
 	}
 }
diff --git a/Engine/Plugins/Enterprise/GLTFExporter/Source/GLTFExporter/Private/Materials/GLTFProxyMaterialInfo.cpp b/Engine/Plugins/Enterprise/GLTFExporter/Source/GLTFExporter/Private/Materials/GLTFProxyMaterialInfo.cpp
--- a/Engine/Plugins/Enterprise/GLTFExporter/Source/GLTFExporter/Private/Materials/GLTFProxyMaterialInfo.cpp
+++ b/Engine/Plugins/Enterprise/GLTFExporter/Source/GLTFExporter/Private/Materials/GLTFProxyMaterialInfo.cpp
@@ -7,6 +7,7 @@ const TGLTFProxyMaterialParameterInfo<FLinearColor> FGLTFProxyMaterialInfo::Base
 
 const FGLTFProxyMaterialTextureParameterInfo FGLTFProxyMaterialInfo::Emissive = { TEXT("Emissive") };
 const TGLTFProxyMaterialParameterInfo<FLinearColor> FGLTFProxyMaterialInfo::EmissiveFactor = { TEXT("Emissive Factor") };
+const TGLTFProxyMaterialParameterInfo<float> FGLTFProxyMaterialInfo::EmissiveStrength = { TEXT("Emissive Strength") };
 
 const FGLTFProxyMaterialTextureParameterInfo FGLTFProxyMaterialInfo::MetallicRoughness = { TEXT("Metallic Roughness") };
 const TGLTFProxyMaterialParameterInfo<float> FGLTFProxyMaterialInfo::MetallicFactor = { TEXT("Metallic Factor") };
diff --git a/Engine/Plugins/Enterprise/GLTFExporter/Source/GLTFExporter/Private/Materials/GLTFProxyMaterialInfo.h b/Engine/Plugins/Enterprise/GLTFExporter/Source/GLTFExporter/Private/Materials/GLTFProxyMaterialInfo.h
--- a/Engine/Plugins/Enterprise/GLTFExporter/Source/GLTFExporter/Private/Materials/GLTFProxyMaterialInfo.h
+++ b/Engine/Plugins/Enterprise/GLTFExporter/Source/GLTFExporter/Private/Materials/GLTFProxyMaterialInfo.h
@@ -13,6 +13,7 @@ public:
 
 	static const FGLTFProxyMaterialTextureParameterInfo Emissive;
 	static const TGLTFProxyMaterialParameterInfo<FLinearColor> EmissiveFactor;
+	static const TGLTFProxyMaterialParameterInfo<float> EmissiveStrength;
 
 	static const FGLTFProxyMaterialTextureParameterInfo MetallicRoughness;
 	static const TGLTFProxyMaterialParameterInfo<float> MetallicFactor;
diff --git a/Engine/Plugins/Enterprise/GLTFExporter/Source/GLTFExporter/Private/Options/GLTFExportOptions.cpp b/Engine/Plugins/Enterprise/GLTFExporter/Source/GLTFExporter/Private/Options/GLTFExportOptions.cpp
--- a/Engine/Plugins/Enterprise/GLTFExporter/Source/GLTFExporter/Private/Options/GLTFExportOptions.cpp
+++ b/Engine/Plugins/Enterprise/GLTFExporter/Source/GLTFExporter/Private/Options/GLTFExportOptions.cpp
@@ -17,6 +17,7 @@ void UGLTFExportOptions::ResetToDefault()
 	bExportProxyMaterials = true;
 	bExportUnlitMaterials = true;
 	bExportClearCoatMaterials = true;
+	bExportEmissiveStrength = true;
 	BakeMaterialInputs = EGLTFMaterialBakeMode::UseMeshData;
 	DefaultMaterialBakeSize = EGLTFMaterialBakeSizePOT::POT_1024;
 	DefaultMaterialBakeFilter = TF_Trilinear;
diff --git a/Engine/Plugins/Enterprise/GLTFExporter/Source/GLTFExporter/Private/Tasks/GLTFDelayedMaterialTasks.cpp b/Engine/Plugins/Enterprise/GLTFExporter/Source/GLTFExporter/Private/Tasks/GLTFDelayedMaterialTasks.cpp
--- a/Engine/Plugins/Enterprise/GLTFExporter/Source/GLTFExporter/Private/Tasks/GLTFDelayedMaterialTasks.cpp
+++ b/Engine/Plugins/Enterprise/GLTFExporter/Source/GLTFExporter/Private/Tasks/GLTFDelayedMaterialTasks.cpp
@@ -214,17 +214,18 @@ FString FGLTFDelayedMaterialTask::GetBakedTextureName(const FString& PropertyNam
 
 void FGLTFDelayedMaterialTask::GetProxyParameters(FGLTFJsonMaterial& OutMaterial) const
 {
-	GetProxyParameter(FGLTFProxyMaterialInfo::BaseColorFactor, OutMaterial.PBRMetallicRoughness.BaseColorFactor);
 	GetProxyParameter(FGLTFProxyMaterialInfo::BaseColor, OutMaterial.PBRMetallicRoughness.BaseColorTexture);
+	GetProxyParameter(FGLTFProxyMaterialInfo::BaseColorFactor, OutMaterial.PBRMetallicRoughness.BaseColorFactor);
 
 	if (OutMaterial.ShadingModel == EGLTFJsonShadingModel::Default || OutMaterial.ShadingModel == EGLTFJsonShadingModel::ClearCoat)
 	{
-		GetProxyParameter(FGLTFProxyMaterialInfo::EmissiveFactor, OutMaterial.EmissiveFactor);
 		GetProxyParameter(FGLTFProxyMaterialInfo::Emissive, OutMaterial.EmissiveTexture);
+		GetProxyParameter(FGLTFProxyMaterialInfo::EmissiveFactor, OutMaterial.EmissiveFactor);
+		GetProxyParameter(FGLTFProxyMaterialInfo::EmissiveStrength, OutMaterial.EmissiveStrength);
 
+		GetProxyParameter(FGLTFProxyMaterialInfo::MetallicRoughness, OutMaterial.PBRMetallicRoughness.MetallicRoughnessTexture);
 		GetProxyParameter(FGLTFProxyMaterialInfo::MetallicFactor, OutMaterial.PBRMetallicRoughness.MetallicFactor);
 		GetProxyParameter(FGLTFProxyMaterialInfo::RoughnessFactor, OutMaterial.PBRMetallicRoughness.RoughnessFactor);
-		GetProxyParameter(FGLTFProxyMaterialInfo::MetallicRoughness, OutMaterial.PBRMetallicRoughness.MetallicRoughnessTexture);
 
 		if (HasProxyParameter(FGLTFProxyMaterialInfo::Normal.Texture))
 		{
@@ -236,24 +237,24 @@ void FGLTFDelayedMaterialTask::GetProxyParameters(FGLTFJsonMaterial& OutMaterial
 			}
 			else
 			{
-				GetProxyParameter(FGLTFProxyMaterialInfo::NormalScale, OutMaterial.NormalTexture.Scale);
 				GetProxyParameter(FGLTFProxyMaterialInfo::Normal, OutMaterial.NormalTexture);
+				GetProxyParameter(FGLTFProxyMaterialInfo::NormalScale, OutMaterial.NormalTexture.Scale);
 			}
 		}
 
-		GetProxyParameter(FGLTFProxyMaterialInfo::OcclusionStrength, OutMaterial.OcclusionTexture.Strength);
 		GetProxyParameter(FGLTFProxyMaterialInfo::Occlusion, OutMaterial.OcclusionTexture);
+		GetProxyParameter(FGLTFProxyMaterialInfo::OcclusionStrength, OutMaterial.OcclusionTexture.Strength);
 
 		if (OutMaterial.ShadingModel == EGLTFJsonShadingModel::ClearCoat)
 		{
-			GetProxyParameter(FGLTFProxyMaterialInfo::ClearCoatFactor, OutMaterial.ClearCoat.ClearCoatFactor);
 			GetProxyParameter(FGLTFProxyMaterialInfo::ClearCoat, OutMaterial.ClearCoat.ClearCoatTexture);
+			GetProxyParameter(FGLTFProxyMaterialInfo::ClearCoatFactor, OutMaterial.ClearCoat.ClearCoatFactor);
 
-			GetProxyParameter(FGLTFProxyMaterialInfo::ClearCoatRoughnessFactor, OutMaterial.ClearCoat.ClearCoatRoughnessFactor);
 			GetProxyParameter(FGLTFProxyMaterialInfo::ClearCoatRoughness, OutMaterial.ClearCoat.ClearCoatRoughnessTexture);
+			GetProxyParameter(FGLTFProxyMaterialInfo::ClearCoatRoughnessFactor, OutMaterial.ClearCoat.ClearCoatRoughnessFactor);
 
-			GetProxyParameter(FGLTFProxyMaterialInfo::ClearCoatNormalScale, OutMaterial.ClearCoat.ClearCoatNormalTexture.Scale);
 			GetProxyParameter(FGLTFProxyMaterialInfo::ClearCoatNormal, OutMaterial.ClearCoat.ClearCoatNormalTexture);
+			GetProxyParameter(FGLTFProxyMaterialInfo::ClearCoatNormalScale, OutMaterial.ClearCoat.ClearCoatNormalTexture.Scale);
 		}
 	}
 }
@@ -579,7 +580,7 @@ bool FGLTFDelayedMaterialTask::TryGetBaseColorAndOpacity(FGLTFJsonPBRMetallicRou
 	OutPBRParams.BaseColorTexture.TexCoord = CombinedTexCoord;
 	OutPBRParams.BaseColorTexture.Index = CombinedTexture;
 
-	// TODO: add support for KHR_materials_emissive_strength
+	// TODO: add warning if BaseColorScale exceeds 1.0 (maybe suggest disabling unlit materials?)
 	OutPBRParams.BaseColorFactor = FGLTFCoreUtilities::ConvertColor(FLinearColor(BaseColorScale, BaseColorScale, BaseColorScale));
 
 	return true;
@@ -817,51 +818,82 @@ bool FGLTFDelayedMaterialTask::TryGetClearCoatRoughness(FGLTFJsonClearCoatExtens
 
 bool FGLTFDelayedMaterialTask::TryGetEmissive(FGLTFJsonMaterial& OutMaterial, const FMaterialPropertyEx& EmissiveProperty)
 {
-	// TODO: right now we allow EmissiveFactor to be > 1.0 to support very bright emission, although it's not valid according to the glTF standard.
-	// We may want to change this behaviour and store factors above 1.0 using a custom extension instead.
-
-	if (TryGetConstantColor(OutMaterial.EmissiveFactor, MP_EmissiveColor))
-	{
-		return true;
-	}
-
-	if (TryGetSourceTexture(OutMaterial.EmissiveTexture, EmissiveProperty, DefaultColorInputMasks))
+	FLinearColor ConstantColor;
+	if (TryGetConstantColor(ConstantColor, MP_EmissiveColor))
 	{
-		OutMaterial.EmissiveFactor = FGLTFJsonColor3::White;	// make sure texture is not multiplied with black
-		return true;
-	}
+		const float EmissiveStrength = ConstantColor.GetMax();
+		if (EmissiveStrength > 1.0)
+		{
+			ConstantColor *= 1.0f / EmissiveStrength;
+			OutMaterial.EmissiveStrength = EmissiveStrength;
+		}
 
-	if (Builder.ExportOptions->BakeMaterialInputs == EGLTFMaterialBakeMode::Disabled)
-	{
-		Builder.LogWarning(FString::Printf(
-			TEXT("%s for material %s needs to bake, but material baking is disabled by export options"),
-			*EmissiveProperty.ToString(),
-			*Material->GetName()));
-		return false;
+		OutMaterial.EmissiveFactor = FGLTFCoreUtilities::ConvertColor3(ConstantColor);
 	}
-
-	FGLTFPropertyBakeOutput PropertyBakeOutput = BakeMaterialProperty(EmissiveProperty, OutMaterial.EmissiveTexture.TexCoord);
-	const float EmissiveScale = PropertyBakeOutput.EmissiveScale;
-
-	if (PropertyBakeOutput.bIsConstant)
+	else if (TryGetSourceTexture(OutMaterial.EmissiveTexture, EmissiveProperty, DefaultColorInputMasks))
 	{
-		const FLinearColor EmissiveColor = PropertyBakeOutput.ConstantValue;
-		OutMaterial.EmissiveFactor = FGLTFCoreUtilities::ConvertColor3(EmissiveColor * EmissiveScale);
+		// TODO: what if texture is HDR and one or more pixel components exceeds 1.0? Should we process it and add emissive strength?
+		OutMaterial.EmissiveFactor = FGLTFJsonColor3::White; // make sure texture is not multiplied with black
 	}
 	else
 	{
-		if (Builder.ExportOptions->TextureImageFormat == EGLTFTextureImageFormat::None)
+		if (Builder.ExportOptions->BakeMaterialInputs == EGLTFMaterialBakeMode::Disabled)
 		{
-			OutMaterial.EmissiveTexture.Index = nullptr;
-			return true;
+			Builder.LogWarning(FString::Printf(
+				TEXT("%s for material %s needs to bake, but material baking is disabled by export options"),
+				*EmissiveProperty.ToString(),
+				*Material->GetName()));
+			return false;
 		}
 
-		if (!StoreBakedPropertyTexture(OutMaterial.EmissiveTexture, PropertyBakeOutput, TEXT("Emissive")))
+		FGLTFPropertyBakeOutput PropertyBakeOutput = BakeMaterialProperty(EmissiveProperty, OutMaterial.EmissiveTexture.TexCoord);
+
+		if (PropertyBakeOutput.bIsConstant)
 		{
-			return false;
+			FLinearColor EmissiveColor = PropertyBakeOutput.ConstantValue;
+			const float EmissiveScale = PropertyBakeOutput.EmissiveScale;
+
+			if (EmissiveScale < 1.0f)
+			{
+				EmissiveColor *= EmissiveScale;
+			}
+			else
+			{
+				OutMaterial.EmissiveStrength = EmissiveScale;
+			}
+
+			OutMaterial.EmissiveFactor = FGLTFCoreUtilities::ConvertColor3(EmissiveColor);
 		}
+		else
+		{
+			if (Builder.ExportOptions->TextureImageFormat == EGLTFTextureImageFormat::None)
+			{
+				OutMaterial.EmissiveTexture.Index = nullptr;
+				return true;
+			}
+
+			if (!StoreBakedPropertyTexture(OutMaterial.EmissiveTexture, PropertyBakeOutput, TEXT("Emissive")))
+			{
+				return false;
+			}
+
+			const float EmissiveScale = PropertyBakeOutput.EmissiveScale;
+			if (EmissiveScale < 1.0f)
+			{
+				OutMaterial.EmissiveFactor = { EmissiveScale, EmissiveScale, EmissiveScale };
+			}
+			else
+			{
+				OutMaterial.EmissiveFactor = FGLTFJsonColor3::White;
+				OutMaterial.EmissiveStrength = EmissiveScale;
+			}
+		}
+	}
 
-		OutMaterial.EmissiveFactor = { EmissiveScale, EmissiveScale, EmissiveScale };
+	if (OutMaterial.EmissiveStrength > 1.0f && !Builder.ExportOptions->bExportEmissiveStrength)
+	{
+		// TODO: add warning about clamping emissive strength?
+		OutMaterial.EmissiveStrength = 1.0f;
 	}
 
 	return true;
diff --git a/Engine/Plugins/Enterprise/GLTFExporter/Source/GLTFExporter/Private/Utilities/GLTFJsonUtilities.cpp b/Engine/Plugins/Enterprise/GLTFExporter/Source/GLTFExporter/Private/Utilities/GLTFJsonUtilities.cpp
--- a/Engine/Plugins/Enterprise/GLTFExporter/Source/GLTFExporter/Private/Utilities/GLTFJsonUtilities.cpp
+++ b/Engine/Plugins/Enterprise/GLTFExporter/Source/GLTFExporter/Private/Utilities/GLTFJsonUtilities.cpp
@@ -6,12 +6,13 @@ const TCHAR* FGLTFJsonUtilities::GetValue(EGLTFJsonExtension Enum)
 {
 	switch (Enum)
 	{
-		case EGLTFJsonExtension::KHR_LightsPunctual:      return TEXT("KHR_lights_punctual");
-		case EGLTFJsonExtension::KHR_MaterialsClearCoat:  return TEXT("KHR_materials_clearcoat");
-		case EGLTFJsonExtension::KHR_MaterialsUnlit:      return TEXT("KHR_materials_unlit");
-		case EGLTFJsonExtension::KHR_MaterialsVariants:   return TEXT("KHR_materials_variants");
-		case EGLTFJsonExtension::KHR_MeshQuantization:    return TEXT("KHR_mesh_quantization");
-		case EGLTFJsonExtension::KHR_TextureTransform:    return TEXT("KHR_texture_transform");
+		case EGLTFJsonExtension::KHR_LightsPunctual:            return TEXT("KHR_lights_punctual");
+		case EGLTFJsonExtension::KHR_MaterialsClearCoat:        return TEXT("KHR_materials_clearcoat");
+		case EGLTFJsonExtension::KHR_MaterialsEmissiveStrength: return TEXT("KHR_materials_emissive_strength");
+		case EGLTFJsonExtension::KHR_MaterialsUnlit:            return TEXT("KHR_materials_unlit");
+		case EGLTFJsonExtension::KHR_MaterialsVariants:         return TEXT("KHR_materials_variants");
+		case EGLTFJsonExtension::KHR_MeshQuantization:          return TEXT("KHR_mesh_quantization");
+		case EGLTFJsonExtension::KHR_TextureTransform:          return TEXT("KHR_texture_transform");
 		default:
 			checkNoEntry();
 			return TEXT("");
diff --git a/Engine/Plugins/Enterprise/GLTFExporter/Source/GLTFExporter/Public/Json/GLTFJsonEnums.h b/Engine/Plugins/Enterprise/GLTFExporter/Source/GLTFExporter/Public/Json/GLTFJsonEnums.h
--- a/Engine/Plugins/Enterprise/GLTFExporter/Source/GLTFExporter/Public/Json/GLTFJsonEnums.h
+++ b/Engine/Plugins/Enterprise/GLTFExporter/Source/GLTFExporter/Public/Json/GLTFJsonEnums.h
@@ -6,6 +6,7 @@ enum class EGLTFJsonExtension
 {
 	KHR_LightsPunctual,
 	KHR_MaterialsClearCoat,
+	KHR_MaterialsEmissiveStrength,
 	KHR_MaterialsUnlit,
 	KHR_MaterialsVariants,
 	KHR_MeshQuantization,
diff --git a/Engine/Plugins/Enterprise/GLTFExporter/Source/GLTFExporter/Public/Json/GLTFJsonMaterial.h b/Engine/Plugins/Enterprise/GLTFExporter/Source/GLTFExporter/Public/Json/GLTFJsonMaterial.h
--- a/Engine/Plugins/Enterprise/GLTFExporter/Source/GLTFExporter/Public/Json/GLTFJsonMaterial.h
+++ b/Engine/Plugins/Enterprise/GLTFExporter/Source/GLTFExporter/Public/Json/GLTFJsonMaterial.h
@@ -96,6 +96,7 @@ struct GLTFEXPORTER_API FGLTFJsonMaterial : IGLTFJsonIndexedObject
 
 	FGLTFJsonTextureInfo EmissiveTexture;
 	FGLTFJsonColor3 EmissiveFactor;
+	float EmissiveStrength;
 
 	EGLTFJsonAlphaMode AlphaMode;
 	float AlphaCutoff;
@@ -114,6 +115,7 @@ protected:
 		: IGLTFJsonIndexedObject(Index)
 		, ShadingModel(EGLTFJsonShadingModel::Default)
 		, EmissiveFactor(FGLTFJsonColor3::Black)
+		, EmissiveStrength(1.0f)
 		, AlphaMode(EGLTFJsonAlphaMode::Opaque)
 		, AlphaCutoff(0.5f)
 		, DoubleSided(false)
diff --git a/Engine/Plugins/Enterprise/GLTFExporter/Source/GLTFExporter/Public/Options/GLTFExportOptions.h b/Engine/Plugins/Enterprise/GLTFExporter/Source/GLTFExporter/Public/Options/GLTFExportOptions.h
--- a/Engine/Plugins/Enterprise/GLTFExporter/Source/GLTFExporter/Public/Options/GLTFExportOptions.h
+++ b/Engine/Plugins/Enterprise/GLTFExporter/Source/GLTFExporter/Public/Options/GLTFExportOptions.h
@@ -57,10 +57,14 @@ class GLTFEXPORTER_API UGLTFExportOptions : public UObject
 	UPROPERTY(EditAnywhere, BlueprintReadWrite, Config, Category = Material)
 	bool bExportUnlitMaterials;
 
-	/** If enabled, materials with shading model clear coat will be properly exported. Uses extension KHR_materials_clearcoat, which is not supported by all glTF viewers. */
+	/** If enabled, materials with shading model clear coat will be properly exported. Uses extension KHR_materials_clearcoat. */
 	UPROPERTY(EditAnywhere, BlueprintReadWrite, Config, Category = Material)
 	bool bExportClearCoatMaterials;
 
+	/** If enabled, allows materials to have an emissive factor that exceeds the standard range [0.0, 1.0]. Uses extension KHR_materials_emissive_strength. */
+	UPROPERTY(EditAnywhere, BlueprintReadWrite, Config, Category = Material)
+	bool bExportEmissiveStrength;
+
 	/** Bake mode determining if and how a material input is baked out to a texture. Baking is only used for non-trivial material inputs (i.e. not simple texture or constant expressions). */
 	UPROPERTY(EditAnywhere, BlueprintReadWrite, Config, Category = Material)
 	EGLTFMaterialBakeMode BakeMaterialInputs;
