Improve texture coverage area when baking using mesh data in glTF exporter

diff --git a/Engine/Plugins/Enterprise/GLTFExporter/Source/GLTFExporter/Private/Converters/GLTFMaterialConverters.cpp b/Engine/Plugins/Enterprise/GLTFExporter/Source/GLTFExporter/Private/Converters/GLTFMaterialConverters.cpp
--- a/Engine/Plugins/Enterprise/GLTFExporter/Source/GLTFExporter/Private/Converters/GLTFMaterialConverters.cpp
+++ b/Engine/Plugins/Enterprise/GLTFExporter/Source/GLTFExporter/Private/Converters/GLTFMaterialConverters.cpp
@@ -48,7 +48,7 @@ FGLTFJsonMaterial* FGLTFMaterialConverter::Convert(const UMaterialInterface* Mat
 	if (Material != FGLTFMaterialUtilities::GetDefaultMaterial())
 	{
 		FGLTFJsonMaterial* JsonMaterial = Builder.AddMaterial();
-		Builder.ScheduleSlowTask<FGLTFDelayedMaterialTask>(Builder, UVOverlapChecker, Material, MeshData, SectionIndices, JsonMaterial);
+		Builder.ScheduleSlowTask<FGLTFDelayedMaterialTask>(Builder, UVOverlapChecker, UVBoundsCalculator, Material, MeshData, SectionIndices, JsonMaterial);
 		return JsonMaterial;
 	}
 
diff --git a/Engine/Plugins/Enterprise/GLTFExporter/Source/GLTFExporter/Private/Converters/GLTFMaterialUtilities.cpp b/Engine/Plugins/Enterprise/GLTFExporter/Source/GLTFExporter/Private/Converters/GLTFMaterialUtilities.cpp
--- a/Engine/Plugins/Enterprise/GLTFExporter/Source/GLTFExporter/Private/Converters/GLTFMaterialUtilities.cpp
+++ b/Engine/Plugins/Enterprise/GLTFExporter/Source/GLTFExporter/Private/Converters/GLTFMaterialUtilities.cpp
@@ -15,7 +15,6 @@
 #endif
 #include "Engine/RendererSettings.h"
 #include "Modules/ModuleManager.h"
-#include "Materials/MaterialInstanceConstant.h"
 #include "Materials/MaterialExpressionCustomOutput.h"
 #include "Materials/MaterialExpressionTextureCoordinate.h"
 
@@ -109,11 +108,11 @@ UMaterialExpressionCustomOutput* FGLTFMaterialUtilities::GetCustomOutputByName(c
 	return nullptr;
 }
 
-FGLTFPropertyBakeOutput FGLTFMaterialUtilities::BakeMaterialProperty(const FIntPoint& OutputSize, const FMaterialPropertyEx& Property, const UMaterialInterface* Material, int32 TexCoord, const FGLTFMeshData* MeshData, const FGLTFIndexArray& MeshSectionIndices, bool bFillAlpha, bool bAdjustNormalmaps)
+FGLTFPropertyBakeOutput FGLTFMaterialUtilities::BakeMaterialProperty(const FIntPoint& OutputSize, const FMaterialPropertyEx& Property, const UMaterialInterface* Material, const FBox2f& TexCoordBounds, int32 TexCoordIndex, const FGLTFMeshData* MeshData, const FGLTFIndexArray& MeshSectionIndices, bool bFillAlpha, bool bAdjustNormalmaps)
 {
 	FMeshData MeshSet;
-	MeshSet.TextureCoordinateBox = { { 0.0f, 0.0f }, { 1.0f, 1.0f } };
-	MeshSet.TextureCoordinateIndex = TexCoord;
+	MeshSet.TextureCoordinateBox = FBox2D(TexCoordBounds);
+	MeshSet.TextureCoordinateIndex = TexCoordIndex;
 	MeshSet.MaterialIndices = MeshSectionIndices; // NOTE: MaterialIndices is actually section indices
 	if (MeshData != nullptr)
 	{
diff --git a/Engine/Plugins/Enterprise/GLTFExporter/Source/GLTFExporter/Private/Converters/GLTFMaterialUtilities.h b/Engine/Plugins/Enterprise/GLTFExporter/Source/GLTFExporter/Private/Converters/GLTFMaterialUtilities.h
--- a/Engine/Plugins/Enterprise/GLTFExporter/Source/GLTFExporter/Private/Converters/GLTFMaterialUtilities.h
+++ b/Engine/Plugins/Enterprise/GLTFExporter/Source/GLTFExporter/Private/Converters/GLTFMaterialUtilities.h
@@ -16,7 +16,13 @@ class UMaterialExpressionCustomOutput;
 struct FGLTFPropertyBakeOutput
 {
 	FGLTFPropertyBakeOutput(const FMaterialPropertyEx& Property, EPixelFormat PixelFormat, const TGLTFSharedArray<FColor>& Pixels, FIntPoint Size, float EmissiveScale, bool bIsSRGB)
-		: Property(Property), PixelFormat(PixelFormat), Pixels(Pixels), Size(Size), EmissiveScale(EmissiveScale), bIsSRGB(bIsSRGB), bIsConstant(false)
+		: Property(Property)
+		, PixelFormat(PixelFormat)
+		, Pixels(Pixels)
+		, Size(Size)
+		, EmissiveScale(EmissiveScale)
+		, bIsSRGB(bIsSRGB)
+		, bIsConstant(false)
 	{}
 
 	const FMaterialPropertyEx& Property;
@@ -56,7 +62,7 @@ struct FGLTFMaterialUtilities
 
 	static UMaterialExpressionCustomOutput* GetCustomOutputByName(const UMaterialInterface* Material, const FString& FunctionName);
 
-	static FGLTFPropertyBakeOutput BakeMaterialProperty(const FIntPoint& OutputSize, const FMaterialPropertyEx& Property, const UMaterialInterface* Material, int32 TexCoord, const FGLTFMeshData* MeshData = nullptr, const FGLTFIndexArray& MeshSectionIndices = {}, bool bFillAlpha = true, bool bAdjustNormalmaps = true);
+	static FGLTFPropertyBakeOutput BakeMaterialProperty(const FIntPoint& OutputSize, const FMaterialPropertyEx& Property, const UMaterialInterface* Material, const FBox2f& TexCoordBounds, int32 TexCoordIndex, const FGLTFMeshData* MeshData, const FGLTFIndexArray& MeshSectionIndices, bool bFillAlpha, bool bAdjustNormalmaps);
 
 	static FGLTFJsonTexture* AddTexture(FGLTFConvertBuilder& Builder, TGLTFSharedArray<FColor>& Pixels, const FIntPoint& TextureSize, bool bIgnoreAlpha, bool bIsNormalMap, const FString& TextureName, TextureAddress TextureAddress, TextureFilter TextureFilter);
 
diff --git a/Engine/Plugins/Enterprise/GLTFExporter/Source/GLTFExporter/Private/Converters/GLTFUVBoundsCalculator.cpp b/Engine/Plugins/Enterprise/GLTFExporter/Source/GLTFExporter/Private/Converters/GLTFUVBoundsCalculator.cpp
new file mode 100644
--- /dev/null
+++ b/Engine/Plugins/Enterprise/GLTFExporter/Source/GLTFExporter/Private/Converters/GLTFUVBoundsCalculator.cpp
@@ -0,0 +1,37 @@
+// Copyright Epic Games, Inc. All Rights Reserved.
+
+#include "Converters/GLTFUVBoundsCalculator.h"
+#if WITH_EDITOR
+#include "StaticMeshAttributes.h"
+#include "MeshDescription.h"
+#endif
+
+FBox2f FGLTFUVBoundsCalculator::Convert(const FMeshDescription* Description, FGLTFIndexArray SectionIndices, int32 TexCoord)
+{
+	FBox2f UVBounds(ForceInit);
+
+#if WITH_EDITOR
+	if (Description != nullptr)
+	{
+		const TVertexInstanceAttributesConstRef<FVector2f> UVs =
+			Description->VertexInstanceAttributes().GetAttributesRef<FVector2f>(MeshAttribute::VertexInstance::TextureCoordinate);
+
+		for (const int32 SectionIndex : SectionIndices)
+		{
+			for (const FPolygonID PolygonID : Description->GetPolygonGroupPolygonIDs(FPolygonGroupID(SectionIndex)))
+			{
+				for (const FTriangleID TriangleID : Description->GetPolygonTriangles(PolygonID))
+				{
+					for (const FVertexInstanceID TriangleVertexInstanceID : Description->GetTriangleVertexInstances(TriangleID))
+					{
+						const FVector2f UV = UVs.Get(TriangleVertexInstanceID, TexCoord);
+						UVBounds += UV;
+					}
+				}
+			}
+		}
+	}
+#endif
+
+	return UVBounds;
+}
diff --git a/Engine/Plugins/Enterprise/GLTFExporter/Source/GLTFExporter/Private/Tasks/GLTFDelayedMaterialTasks.cpp b/Engine/Plugins/Enterprise/GLTFExporter/Source/GLTFExporter/Private/Tasks/GLTFDelayedMaterialTasks.cpp
--- a/Engine/Plugins/Enterprise/GLTFExporter/Source/GLTFExporter/Private/Tasks/GLTFDelayedMaterialTasks.cpp
+++ b/Engine/Plugins/Enterprise/GLTFExporter/Source/GLTFExporter/Private/Tasks/GLTFDelayedMaterialTasks.cpp
@@ -521,8 +521,8 @@ bool FGLTFDelayedMaterialTask::TryGetBaseColorAndOpacity(FGLTFJsonPBRMetallicRou
 	const TextureAddress TextureAddress = Builder.GetBakeTilingForMaterialProperty(Material, GetPropertyGroup(BaseColorProperty));
 	const TextureFilter TextureFilter = Builder.GetBakeFilterForMaterialProperty(Material, GetPropertyGroup(BaseColorProperty));
 
-	const FGLTFPropertyBakeOutput BaseColorBakeOutput = BakeMaterialProperty(BaseColorProperty, BaseColorTexCoord, TextureSize, false);
-	const FGLTFPropertyBakeOutput OpacityBakeOutput = BakeMaterialProperty(OpacityProperty, OpacityTexCoord, TextureSize, false);
+	const FGLTFPropertyBakeOutput BaseColorBakeOutput = BakeMaterialProperty(BaseColorProperty, BaseColorTexCoord, BaseColorTransform, TextureSize, false);
+	const FGLTFPropertyBakeOutput OpacityBakeOutput = BakeMaterialProperty(OpacityProperty, OpacityTexCoord, OpacityTransform, TextureSize, false);
 	const float BaseColorScale = BaseColorProperty == MP_EmissiveColor ? BaseColorBakeOutput.EmissiveScale : 1;
 
 	// Detect when both baked properties are constants, which means we can avoid exporting a texture
@@ -541,15 +541,18 @@ bool FGLTFDelayedMaterialTask::TryGetBaseColorAndOpacity(FGLTFJsonPBRMetallicRou
 	}
 
 	int32 CombinedTexCoord;
+	FGLTFJsonTextureTransform CombinedTransform;
 
 	// If one is constant, it means we can use the other's texture coords
 	if (BaseColorBakeOutput.bIsConstant || BaseColorTexCoord == OpacityTexCoord)
 	{
 		CombinedTexCoord = OpacityTexCoord;
+		CombinedTransform = OpacityTransform; // If texcoord the same, then transform also the same
 	}
 	else if (OpacityBakeOutput.bIsConstant)
 	{
 		CombinedTexCoord = BaseColorTexCoord;
+		CombinedTransform = BaseColorTransform;
 	}
 	else
 	{
@@ -577,8 +580,9 @@ bool FGLTFDelayedMaterialTask::TryGetBaseColorAndOpacity(FGLTFJsonPBRMetallicRou
 		TextureAddress,
 		TextureFilter);
 
-	OutPBRParams.BaseColorTexture.TexCoord = CombinedTexCoord;
 	OutPBRParams.BaseColorTexture.Index = CombinedTexture;
+	OutPBRParams.BaseColorTexture.TexCoord = CombinedTexCoord;
+	OutPBRParams.BaseColorTexture.Transform = CombinedTransform;
 
 	// TODO: add warning if BaseColorScale exceeds 1.0 (maybe suggest disabling unlit materials?)
 	OutPBRParams.BaseColorFactor = FGLTFCoreUtilities::ConvertColor(FLinearColor(BaseColorScale, BaseColorScale, BaseColorScale));
@@ -640,8 +644,8 @@ bool FGLTFDelayedMaterialTask::TryGetMetallicAndRoughness(FGLTFJsonPBRMetallicRo
 	const TextureAddress TextureAddress = Builder.GetBakeTilingForMaterialProperty(Material, GetPropertyGroup(MetallicProperty));
 	const TextureFilter TextureFilter = Builder.GetBakeFilterForMaterialProperty(Material, GetPropertyGroup(MetallicProperty));
 
-	FGLTFPropertyBakeOutput MetallicBakeOutput = BakeMaterialProperty(MetallicProperty, MetallicTexCoord, TextureSize, false);
-	FGLTFPropertyBakeOutput RoughnessBakeOutput = BakeMaterialProperty(RoughnessProperty, RoughnessTexCoord, TextureSize, false);
+	FGLTFPropertyBakeOutput MetallicBakeOutput = BakeMaterialProperty(MetallicProperty, MetallicTexCoord, MetallicTransform, TextureSize, false);
+	FGLTFPropertyBakeOutput RoughnessBakeOutput = BakeMaterialProperty(RoughnessProperty, RoughnessTexCoord, RoughnessTransform, TextureSize, false);
 
 	// Detect when both baked properties are constants, which means we can use factors and avoid exporting a texture
 	if (MetallicBakeOutput.bIsConstant && RoughnessBakeOutput.bIsConstant)
@@ -657,15 +661,18 @@ bool FGLTFDelayedMaterialTask::TryGetMetallicAndRoughness(FGLTFJsonPBRMetallicRo
 	}
 
 	int32 CombinedTexCoord;
+	FGLTFJsonTextureTransform CombinedTransform;
 
 	// If one is constant, it means we can use the other's texture coords
 	if (MetallicBakeOutput.bIsConstant || MetallicTexCoord == RoughnessTexCoord)
 	{
 		CombinedTexCoord = RoughnessTexCoord;
+		CombinedTransform = RoughnessTransform; // If texcoord the same, then transform also the same
 	}
 	else if (RoughnessBakeOutput.bIsConstant)
 	{
 		CombinedTexCoord = MetallicTexCoord;
+		CombinedTransform = MetallicTransform;
 	}
 	else
 	{
@@ -693,8 +700,9 @@ bool FGLTFDelayedMaterialTask::TryGetMetallicAndRoughness(FGLTFJsonPBRMetallicRo
 		TextureAddress,
 		TextureFilter);
 
-	OutPBRParams.MetallicRoughnessTexture.TexCoord = CombinedTexCoord;
 	OutPBRParams.MetallicRoughnessTexture.Index = CombinedTexture;
+	OutPBRParams.MetallicRoughnessTexture.TexCoord = CombinedTexCoord;
+	OutPBRParams.MetallicRoughnessTexture.Transform = CombinedTransform;
 
 	return true;
 }
@@ -755,8 +763,8 @@ bool FGLTFDelayedMaterialTask::TryGetClearCoatRoughness(FGLTFJsonClearCoatExtens
 	const TextureAddress TextureAddress = Builder.GetBakeTilingForMaterialProperty(Material,GetPropertyGroup(IntensityProperty));
 	const TextureFilter TextureFilter = Builder.GetBakeFilterForMaterialProperty(Material, GetPropertyGroup(IntensityProperty));
 
-	const FGLTFPropertyBakeOutput IntensityBakeOutput = BakeMaterialProperty(IntensityProperty, IntensityTexCoord, TextureSize, false);
-	const FGLTFPropertyBakeOutput RoughnessBakeOutput = BakeMaterialProperty(RoughnessProperty, RoughnessTexCoord, TextureSize, false);
+	const FGLTFPropertyBakeOutput IntensityBakeOutput = BakeMaterialProperty(IntensityProperty, IntensityTexCoord, IntensityTransform, TextureSize, false);
+	const FGLTFPropertyBakeOutput RoughnessBakeOutput = BakeMaterialProperty(RoughnessProperty, RoughnessTexCoord, RoughnessTransform, TextureSize, false);
 
 	// Detect when both baked properties are constants, which means we can use factors and avoid exporting a texture
 	if (IntensityBakeOutput.bIsConstant && RoughnessBakeOutput.bIsConstant)
@@ -772,15 +780,18 @@ bool FGLTFDelayedMaterialTask::TryGetClearCoatRoughness(FGLTFJsonClearCoatExtens
 	}
 
 	int32 CombinedTexCoord;
+	FGLTFJsonTextureTransform CombinedTransform;
 
 	// If one is constant, it means we can use the other's texture coords
 	if (IntensityBakeOutput.bIsConstant || IntensityTexCoord == RoughnessTexCoord)
 	{
 		CombinedTexCoord = RoughnessTexCoord;
+		CombinedTransform = RoughnessTransform; // If texcoord the same, then transform also the same
 	}
 	else if (RoughnessBakeOutput.bIsConstant)
 	{
 		CombinedTexCoord = IntensityTexCoord;
+		CombinedTransform = IntensityTransform;
 	}
 	else
 	{
@@ -810,8 +821,10 @@ bool FGLTFDelayedMaterialTask::TryGetClearCoatRoughness(FGLTFJsonClearCoatExtens
 
 	OutExtParams.ClearCoatTexture.Index = CombinedTexture;
 	OutExtParams.ClearCoatTexture.TexCoord = CombinedTexCoord;
+	OutExtParams.ClearCoatTexture.Transform = CombinedTransform;
 	OutExtParams.ClearCoatRoughnessTexture.Index = CombinedTexture;
 	OutExtParams.ClearCoatRoughnessTexture.TexCoord = CombinedTexCoord;
+	OutExtParams.ClearCoatRoughnessTexture.Transform = CombinedTransform;
 
 	return true;
 }
@@ -846,7 +859,7 @@ bool FGLTFDelayedMaterialTask::TryGetEmissive(FGLTFJsonMaterial& OutMaterial, co
 			return false;
 		}
 
-		FGLTFPropertyBakeOutput PropertyBakeOutput = BakeMaterialProperty(EmissiveProperty, OutMaterial.EmissiveTexture.TexCoord);
+		FGLTFPropertyBakeOutput PropertyBakeOutput = BakeMaterialProperty(EmissiveProperty, OutMaterial.EmissiveTexture.TexCoord, OutMaterial.EmissiveTexture.Transform);
 
 		if (PropertyBakeOutput.bIsConstant)
 		{
@@ -1318,7 +1331,7 @@ bool FGLTFDelayedMaterialTask::TryGetBakedMaterialProperty(FGLTFJsonTextureInfo&
 		return false;
 	}
 
-	FGLTFPropertyBakeOutput PropertyBakeOutput = BakeMaterialProperty(Property, OutTexInfo.TexCoord);
+	FGLTFPropertyBakeOutput PropertyBakeOutput = BakeMaterialProperty(Property, OutTexInfo.TexCoord, OutTexInfo.Transform);
 
 	if (PropertyBakeOutput.bIsConstant)
 	{
@@ -1352,7 +1365,7 @@ bool FGLTFDelayedMaterialTask::TryGetBakedMaterialProperty(FGLTFJsonTextureInfo&
 		return false;
 	}
 
-	FGLTFPropertyBakeOutput PropertyBakeOutput = BakeMaterialProperty(Property, OutTexInfo.TexCoord);
+	FGLTFPropertyBakeOutput PropertyBakeOutput = BakeMaterialProperty(Property, OutTexInfo.TexCoord, OutTexInfo.Transform);
 
 	if (PropertyBakeOutput.bIsConstant)
 	{
@@ -1386,7 +1399,7 @@ inline bool FGLTFDelayedMaterialTask::TryGetBakedMaterialProperty(FGLTFJsonTextu
 		return false;
 	}
 
-	FGLTFPropertyBakeOutput PropertyBakeOutput = BakeMaterialProperty(Property, OutTexInfo.TexCoord);
+	FGLTFPropertyBakeOutput PropertyBakeOutput = BakeMaterialProperty(Property, OutTexInfo.TexCoord, OutTexInfo.Transform);
 
 	if (PropertyBakeOutput.bIsConstant)
 	{
@@ -1420,7 +1433,7 @@ bool FGLTFDelayedMaterialTask::TryGetBakedMaterialProperty(FGLTFJsonTextureInfo&
 		return false;
 	}
 
-	FGLTFPropertyBakeOutput PropertyBakeOutput = BakeMaterialProperty(Property, OutTexInfo.TexCoord);
+	FGLTFPropertyBakeOutput PropertyBakeOutput = BakeMaterialProperty(Property, OutTexInfo.TexCoord, OutTexInfo.Transform);
 
 	if (!PropertyBakeOutput.bIsConstant)
 	{
@@ -1469,14 +1482,17 @@ bool FGLTFDelayedMaterialTask::TryGetBakedMaterialProperty(FGLTFJsonTextureInfo&
 	return true;
 }
 
-FGLTFPropertyBakeOutput FGLTFDelayedMaterialTask::BakeMaterialProperty(const FMaterialPropertyEx& Property, int32& OutTexCoord)
+FGLTFPropertyBakeOutput FGLTFDelayedMaterialTask::BakeMaterialProperty(const FMaterialPropertyEx& Property, int32& OutTexCoord, FGLTFJsonTextureTransform& OutTransform)
 {
 	const FIntPoint TextureSize = Builder.GetBakeSizeForMaterialProperty(Material, GetPropertyGroup(Property));
-	return BakeMaterialProperty(Property, OutTexCoord, TextureSize, true);
+	return BakeMaterialProperty(Property, OutTexCoord, OutTransform, TextureSize, true);
 }
 
-FGLTFPropertyBakeOutput FGLTFDelayedMaterialTask::BakeMaterialProperty(const FMaterialPropertyEx& Property, int32& OutTexCoord, const FIntPoint& TextureSize, bool bFillAlpha)
+FGLTFPropertyBakeOutput FGLTFDelayedMaterialTask::BakeMaterialProperty(const FMaterialPropertyEx& Property, int32& OutTexCoord, FGLTFJsonTextureTransform& OutTransform, const FIntPoint& TextureSize, bool bFillAlpha)
 {
+	const FBox2f DefaultTexCoordBounds = { { 0.0f, 0.0f }, { 1.0f, 1.0f } };
+	FBox2f TexCoordBounds;
+
 	if (MeshData == nullptr)
 	{
 		FGLTFIndexArray TexCoords;
@@ -1500,11 +1516,30 @@ FGLTFPropertyBakeOutput FGLTFDelayedMaterialTask::BakeMaterialProperty(const FMa
 		{
 			OutTexCoord = 0;
 		}
+
+		TexCoordBounds = DefaultTexCoordBounds;
 	}
 	else
 	{
-		OutTexCoord = MeshData->BakeUsingTexCoord;
 		MeshDataBakedProperties.Add(Property);
+		OutTexCoord = MeshData->BakeUsingTexCoord;
+		TexCoordBounds = UVBoundsCalculator.GetOrAdd(&MeshData->Description, SectionIndices, OutTexCoord);
+
+		if (Builder.ExportOptions->bExportTextureTransforms)
+		{
+			const FVector2f Scale = { 1.0f / (TexCoordBounds.Max.X - TexCoordBounds.Min.X), 1.0f / (TexCoordBounds.Max.Y - TexCoordBounds.Min.Y) };
+			const FVector2f Offset = -TexCoordBounds.Min * Scale;
+			OutTransform.Offset = FGLTFCoreUtilities::ConvertUV(Offset);
+			OutTransform.Scale = FGLTFCoreUtilities::ConvertUV(Scale);
+		}
+		else if (!TexCoordBounds.Equals(DefaultTexCoordBounds))
+		{
+			Builder.LogSuggestion(FString::Printf(
+				TEXT("Export and baking of %s for material %s could be improved, by enabling texture transform in export options"),
+				*Property.ToString(),
+				*Material->GetName()));
+			TexCoordBounds = DefaultTexCoordBounds;
+		}
 	}
 
 	// TODO: add support for calculating the ideal resolution to use for baking based on connected (texture) nodes
@@ -1513,6 +1548,7 @@ FGLTFPropertyBakeOutput FGLTFDelayedMaterialTask::BakeMaterialProperty(const FMa
 		TextureSize,
 		Property,
 		Material,
+		TexCoordBounds,
 		OutTexCoord,
 		MeshData,
 		SectionIndices,
diff --git a/Engine/Plugins/Enterprise/GLTFExporter/Source/GLTFExporter/Private/Tasks/GLTFDelayedMaterialTasks.h b/Engine/Plugins/Enterprise/GLTFExporter/Source/GLTFExporter/Private/Tasks/GLTFDelayedMaterialTasks.h
--- a/Engine/Plugins/Enterprise/GLTFExporter/Source/GLTFExporter/Private/Tasks/GLTFDelayedMaterialTasks.h
+++ b/Engine/Plugins/Enterprise/GLTFExporter/Source/GLTFExporter/Private/Tasks/GLTFDelayedMaterialTasks.h
@@ -5,6 +5,7 @@
 #include "Tasks/GLTFDelayedTask.h"
 #include "Builders/GLTFConvertBuilder.h"
 #include "Converters/GLTFUVOverlapChecker.h"
+#include "Converters/GLTFUVBoundsCalculator.h"
 #include "Converters/GLTFMeshData.h"
 #include "Materials/GLTFProxyMaterialParameterInfo.h"
 #include "Materials/Material.h"
@@ -18,10 +19,11 @@ class FGLTFDelayedMaterialTask : public FGLTFDelayedTask
 {
 public:
 
-	FGLTFDelayedMaterialTask(FGLTFConvertBuilder& Builder, FGLTFUVOverlapChecker& UVOverlapChecker, const UMaterialInterface* Material, const FGLTFMeshData* MeshData, FGLTFIndexArray SectionIndices, FGLTFJsonMaterial* JsonMaterial)
+	FGLTFDelayedMaterialTask(FGLTFConvertBuilder& Builder, FGLTFUVOverlapChecker& UVOverlapChecker, FGLTFUVBoundsCalculator& UVBoundsCalculator, const UMaterialInterface* Material, const FGLTFMeshData* MeshData, FGLTFIndexArray SectionIndices, FGLTFJsonMaterial* JsonMaterial)
 		: FGLTFDelayedTask(EGLTFTaskPriority::Material)
 		, Builder(Builder)
 		, UVOverlapChecker(UVOverlapChecker)
+		, UVBoundsCalculator(UVBoundsCalculator)
 		, Material(Material)
 		, MeshData(MeshData)
 		, SectionIndices(SectionIndices)
@@ -40,6 +42,7 @@ private:
 
 	FGLTFConvertBuilder& Builder;
 	FGLTFUVOverlapChecker& UVOverlapChecker;
+	FGLTFUVBoundsCalculator& UVBoundsCalculator;
 	const UMaterialInterface* Material;
 	const FGLTFMeshData* MeshData;
 	const FGLTFIndexArray SectionIndices;
@@ -87,8 +90,8 @@ private:
 	bool TryGetBakedMaterialProperty(FGLTFJsonTextureInfo& OutTexInfo, float& OutConstant, const FMaterialPropertyEx& Property, const FString& PropertyName);
 	bool TryGetBakedMaterialProperty(FGLTFJsonTextureInfo& OutTexInfo, const FMaterialPropertyEx& Property, const FString& PropertyName);
 
-	FGLTFPropertyBakeOutput BakeMaterialProperty(const FMaterialPropertyEx& Property, int32& OutTexCoord);
-	FGLTFPropertyBakeOutput BakeMaterialProperty(const FMaterialPropertyEx& Property, int32& OutTexCoord, const FIntPoint& TextureSize, bool bFillAlpha);
+	FGLTFPropertyBakeOutput BakeMaterialProperty(const FMaterialPropertyEx& Property, int32& OutTexCoord, FGLTFJsonTextureTransform& OutTransform);
+	FGLTFPropertyBakeOutput BakeMaterialProperty(const FMaterialPropertyEx& Property, int32& OutTexCoord, FGLTFJsonTextureTransform& OutTransform, const FIntPoint& TextureSize, bool bFillAlpha);
 
 	bool StoreBakedPropertyTexture(FGLTFJsonTextureInfo& OutTexInfo, FGLTFPropertyBakeOutput& PropertyBakeOutput, const FString& PropertyName) const;
 
diff --git a/Engine/Plugins/Enterprise/GLTFExporter/Source/GLTFExporter/Public/Converters/GLTFMaterialConverters.h b/Engine/Plugins/Enterprise/GLTFExporter/Source/GLTFExporter/Public/Converters/GLTFMaterialConverters.h
--- a/Engine/Plugins/Enterprise/GLTFExporter/Source/GLTFExporter/Public/Converters/GLTFMaterialConverters.h
+++ b/Engine/Plugins/Enterprise/GLTFExporter/Source/GLTFExporter/Public/Converters/GLTFMaterialConverters.h
@@ -8,6 +8,7 @@
 #include "Converters/GLTFMeshData.h"
 #include "Converters/GLTFIndexArray.h"
 #include "Converters/GLTFUVOverlapChecker.h"
+#include "Converters/GLTFUVBoundsCalculator.h"
 #include "Converters/GLTFUVDegenerateChecker.h"
 
 typedef TGLTFConverter<FGLTFJsonMaterial*, const UMaterialInterface*, const FGLTFMeshData*, FGLTFIndexArray> IGLTFMaterialConverter;
@@ -27,5 +28,6 @@ protected:
 private:
 
 	FGLTFUVOverlapChecker UVOverlapChecker;
+	FGLTFUVBoundsCalculator UVBoundsCalculator;
 	FGLTFUVDegenerateChecker UVDegenerateChecker;
 };
diff --git a/Engine/Plugins/Enterprise/GLTFExporter/Source/GLTFExporter/Public/Converters/GLTFUVBoundsCalculator.h b/Engine/Plugins/Enterprise/GLTFExporter/Source/GLTFExporter/Public/Converters/GLTFUVBoundsCalculator.h
new file mode 100644
--- /dev/null
+++ b/Engine/Plugins/Enterprise/GLTFExporter/Source/GLTFExporter/Public/Converters/GLTFUVBoundsCalculator.h
@@ -0,0 +1,15 @@
+// Copyright Epic Games, Inc. All Rights Reserved.
+
+#pragma once
+
+#include "Converters/GLTFConverter.h"
+#include "Converters/GLTFIndexArray.h"
+
+struct FMeshDescription;
+
+class GLTFEXPORTER_API FGLTFUVBoundsCalculator : public TGLTFConverter<FBox2f, const FMeshDescription*, FGLTFIndexArray, int32>
+{
+protected:
+
+	virtual FBox2f Convert(const FMeshDescription* Description, FGLTFIndexArray SectionIndices, int32 TexCoord) override;
+};
diff --git a/Engine/Plugins/Enterprise/GLTFExporter/Source/GLTFExporter/Public/Converters/GLTFUVDegenerateChecker.h b/Engine/Plugins/Enterprise/GLTFExporter/Source/GLTFExporter/Public/Converters/GLTFUVDegenerateChecker.h
--- a/Engine/Plugins/Enterprise/GLTFExporter/Source/GLTFExporter/Public/Converters/GLTFUVDegenerateChecker.h
+++ b/Engine/Plugins/Enterprise/GLTFExporter/Source/GLTFExporter/Public/Converters/GLTFUVDegenerateChecker.h
@@ -2,11 +2,9 @@
 
 #pragma once
 
-#include "CoreMinimal.h"
-#include "Containers/StaticArray.h"
-
 #include "Converters/GLTFConverter.h"
 #include "Converters/GLTFIndexArray.h"
+#include "Containers/StaticArray.h"
 
 struct FMeshDescription;
 
diff --git a/Engine/Plugins/Enterprise/GLTFExporter/Source/GLTFExporter/Public/Options/GLTFExportOptions.h b/Engine/Plugins/Enterprise/GLTFExporter/Source/GLTFExporter/Public/Options/GLTFExportOptions.h
--- a/Engine/Plugins/Enterprise/GLTFExporter/Source/GLTFExporter/Public/Options/GLTFExportOptions.h
+++ b/Engine/Plugins/Enterprise/GLTFExporter/Source/GLTFExporter/Public/Options/GLTFExportOptions.h
@@ -117,7 +117,7 @@ class GLTFEXPORTER_API UGLTFExportOptions : public UObject
 	UPROPERTY(EditAnywhere, BlueprintReadWrite, Config, Category = Texture, Meta = (ClampMin = "0", ClampMax = "100", EditCondition = "TextureImageFormat == EGLTFTextureImageFormat::JPEG"))
 	int32 TextureImageQuality;
 
-	/** If enabled, export UV tiling and un-mirroring settings in a texture coordinate expression node for simple material input expressions. Uses extension KHR_texture_transform. */
+	/** If enabled, export UV offset and scale/tiling used in materials. Uses extension KHR_texture_transform. */
 	UPROPERTY(EditAnywhere, BlueprintReadWrite, Config, Category = Texture, Meta = (EditCondition = "TextureImageFormat != EGLTFTextureImageFormat::None"))
 	bool bExportTextureTransforms;
 
