Check final exported glb file size and give error if exceeding maximum

diff --git a/Engine/Plugins/Enterprise/GLTFExporter/Source/GLTFExporter/Private/Builders/GLTFContainerBuilder.cpp b/Engine/Plugins/Enterprise/GLTFExporter/Source/GLTFExporter/Private/Builders/GLTFContainerBuilder.cpp
--- a/Engine/Plugins/Enterprise/GLTFExporter/Source/GLTFExporter/Private/Builders/GLTFContainerBuilder.cpp
+++ b/Engine/Plugins/Enterprise/GLTFExporter/Source/GLTFExporter/Private/Builders/GLTFContainerBuilder.cpp
@@ -9,24 +9,25 @@ FGLTFContainerBuilder::FGLTFContainerBuilder(const FString& FileName, const UGLT
 {
 }
 
-void FGLTFContainerBuilder::WriteInternalArchive(FArchive& Archive)
+bool FGLTFContainerBuilder::WriteInternalArchive(FArchive& Archive)
 {
 	ProcessSlowTasks();
 
 	if (bIsGLB)
 	{
-		WriteGlbArchive(Archive);
-	}
-	else
-	{
-		WriteJsonArchive(Archive);
+		return WriteGlbArchive(Archive);
 	}
+	
+	return WriteJsonArchive(Archive);
 }
 
 bool FGLTFContainerBuilder::WriteAllFiles(const FString& DirPath, uint32 WriteFlags)
 {
 	FGLTFMemoryArchive Archive;
-	WriteInternalArchive(Archive);
+	if (!WriteInternalArchive(Archive))
+	{
+		return false;
+	}
 
 	const FString FilePath = FPaths::Combine(DirPath, FileName);
 	if (!SaveToFile(FilePath, Archive, WriteFlags))
@@ -54,56 +55,68 @@ void FGLTFContainerBuilder::GetAllFiles(TArray<FString>& OutFilePaths, const FSt
 	GetExternalFiles(OutFilePaths, DirPath);
 }
 
-void FGLTFContainerBuilder::WriteGlbArchive(FArchive& Archive)
+bool FGLTFContainerBuilder::WriteGlbArchive(FArchive& Archive)
 {
 	FGLTFMemoryArchive JsonData;
-	WriteJsonArchive(JsonData);
+	if (!WriteJsonArchive(JsonData))
+	{
+		return false;
+	}
 
 	const TArray64<uint8>* BufferData = GetBufferData();
-	WriteGlb(Archive, JsonData, BufferData);
+	return WriteGlb(Archive, JsonData, BufferData);
 }
 
-void FGLTFContainerBuilder::WriteGlb(FArchive& Archive, const TArray64<uint8>& JsonData, const TArray64<uint8>* BinaryData)
+bool FGLTFContainerBuilder::WriteGlb(FArchive& Archive, const TArray64<uint8>& JsonData, const TArray64<uint8>* BinaryData)
 {
 	constexpr uint32 JsonChunkType = 0x4E4F534A; // "JSON" in ASCII
 	constexpr uint32 BinaryChunkType = 0x004E4942; // "BIN" in ASCII
-	uint32 FileSize =
-		3 * sizeof(uint32) +
-		2 * sizeof(uint32) + GetPaddedChunkSize(JsonData.Num());
+
+	uint64 FileLength =
+		3 * sizeof(uint32) + // header
+		2 * sizeof(uint32) + JsonData.Num() + GetChunkPaddingLength(JsonData.Num()); // JSON chuck
 
 	if (BinaryData != nullptr)
 	{
-		FileSize += 2 * sizeof(uint32) + GetPaddedChunkSize(BinaryData->Num());
+		FileLength += 2 * sizeof(uint32) + BinaryData->Num() + GetChunkPaddingLength(BinaryData->Num()); // BIN chuck
 	}
 
-	WriteHeader(Archive, FileSize);
+	if (FileLength > UINT32_MAX)
+	{
+		LogError(FString::Printf(TEXT("Final file size (%lld) exceedes maximum supported size (%d) in glTF binary container format (.glb)"), FileLength, UINT32_MAX));
+		return false;
+	}
+
+	WriteHeader(Archive, static_cast<uint32>(FileLength));
 	WriteChunk(Archive, JsonChunkType, JsonData, 0x20);
 
 	if (BinaryData != nullptr)
 	{
 		WriteChunk(Archive, BinaryChunkType, *BinaryData, 0x0);
 	}
+
+	return true;
 }
 
-void FGLTFContainerBuilder::WriteHeader(FArchive& Archive, uint32 FileSize)
+void FGLTFContainerBuilder::WriteHeader(FArchive& Archive, uint32 FileLength)
 {
 	constexpr uint32 FileSignature = 0x46546C67; // "glTF" in ASCII
 	constexpr uint32 FileVersion = 2;
 
 	WriteInt(Archive, FileSignature);
 	WriteInt(Archive, FileVersion);
-	WriteInt(Archive, FileSize);
+	WriteInt(Archive, FileLength);
 }
 
-void FGLTFContainerBuilder::WriteChunk(FArchive& Archive, uint32 ChunkType, const TArray64<uint8>& ChunkData, uint8 ChunkTrailingByte)
+void FGLTFContainerBuilder::WriteChunk(FArchive& Archive, uint32 ChunkType, const TArray64<uint8>& ChunkData, uint8 PaddingValue)
 {
-	const uint32 ChunkLength = GetPaddedChunkSize(ChunkData.Num());
-	const uint32 ChunkTrailing = GetTrailingChunkSize(ChunkData.Num());
+	const uint32 PaddingLength = GetChunkPaddingLength(ChunkData.Num());
+	const uint32 ChunkLength = static_cast<uint32>(ChunkData.Num() + PaddingLength);
 
 	WriteInt(Archive, ChunkLength);
 	WriteInt(Archive, ChunkType);
 	WriteData(Archive, ChunkData);
-	WriteFill(Archive, ChunkTrailing, ChunkTrailingByte);
+	WriteFill(Archive, PaddingLength, PaddingValue);
 }
 
 void FGLTFContainerBuilder::WriteInt(FArchive& Archive, uint32 Value)
@@ -125,12 +138,7 @@ void FGLTFContainerBuilder::WriteFill(FArchive& Archive, uint32 Size, uint8 Valu
 	}
 }
 
-uint32 FGLTFContainerBuilder::GetPaddedChunkSize(uint64 Size)
-{
-	return static_cast<uint32>((Size + 3) & ~3);
-}
-
-uint32 FGLTFContainerBuilder::GetTrailingChunkSize(uint64 Size)
+uint32 FGLTFContainerBuilder::GetChunkPaddingLength(uint64 Size)
 {
-	return static_cast<uint32>((4 - (Size & 3)) & 3);
+	return (4 - static_cast<uint32>(Size & 3)) & 3;
 }
diff --git a/Engine/Plugins/Enterprise/GLTFExporter/Source/GLTFExporter/Private/Builders/GLTFJsonBuilder.cpp b/Engine/Plugins/Enterprise/GLTFExporter/Source/GLTFExporter/Private/Builders/GLTFJsonBuilder.cpp
--- a/Engine/Plugins/Enterprise/GLTFExporter/Source/GLTFExporter/Private/Builders/GLTFJsonBuilder.cpp
+++ b/Engine/Plugins/Enterprise/GLTFExporter/Source/GLTFExporter/Private/Builders/GLTFJsonBuilder.cpp
@@ -16,9 +16,10 @@ FGLTFJsonBuilder::FGLTFJsonBuilder(const FString& FileName, const UGLTFExportOpt
 	}
 }
 
-void FGLTFJsonBuilder::WriteJsonArchive(FArchive& Archive)
+bool FGLTFJsonBuilder::WriteJsonArchive(FArchive& Archive)
 {
 	JsonRoot.WriteJson(Archive, !bIsGLB, ExportOptions->bSkipNearDefaultValues ? KINDA_SMALL_NUMBER : 0);
+	return true;
 }
 
 void FGLTFJsonBuilder::AddExtension(EGLTFJsonExtension Extension, bool bIsRequired)
diff --git a/Engine/Plugins/Enterprise/GLTFExporter/Source/GLTFExporter/Public/Builders/GLTFContainerBuilder.h b/Engine/Plugins/Enterprise/GLTFExporter/Source/GLTFExporter/Public/Builders/GLTFContainerBuilder.h
--- a/Engine/Plugins/Enterprise/GLTFExporter/Source/GLTFExporter/Public/Builders/GLTFContainerBuilder.h
+++ b/Engine/Plugins/Enterprise/GLTFExporter/Source/GLTFExporter/Public/Builders/GLTFContainerBuilder.h
@@ -10,7 +10,7 @@ public:
 
 	FGLTFContainerBuilder(const FString& FileName, const UGLTFExportOptions* ExportOptions = nullptr, const TSet<AActor*>& SelectedActors = {});
 
-	void WriteInternalArchive(FArchive& Archive);
+	bool WriteInternalArchive(FArchive& Archive);
 
 	bool WriteAllFiles(const FString& DirPath, uint32 WriteFlags = 0);
 	bool WriteAllFiles(const FString& DirPath, TArray<FString>& OutFilePaths, uint32 WriteFlags = 0);
@@ -19,19 +19,18 @@ public:
 
 protected:
 
-	void WriteGlbArchive(FArchive& Archive);
+	bool WriteGlbArchive(FArchive& Archive);
 
 private:
 
-	static void WriteGlb(FArchive& Archive, const TArray64<uint8>& JsonData, const TArray64<uint8>* BinaryData);
+	bool WriteGlb(FArchive& Archive, const TArray64<uint8>& JsonData, const TArray64<uint8>* BinaryData);
 
 	static void WriteHeader(FArchive& Archive, uint32 FileSize);
-	static void WriteChunk(FArchive& Archive, uint32 ChunkType, const TArray64<uint8>& ChunkData, uint8 ChunkTrailingByte);
+	static void WriteChunk(FArchive& Archive, uint32 ChunkType, const TArray64<uint8>& ChunkData, uint8 PaddingValue);
 
 	static void WriteInt(FArchive& Archive, uint32 Value);
 	static void WriteData(FArchive& Archive, const TArray64<uint8>& Data);
 	static void WriteFill(FArchive& Archive, uint32 Size, uint8 Value);
 
-	static uint32 GetPaddedChunkSize(uint64 Size);
-	static uint32 GetTrailingChunkSize(uint64 Size);
+	static uint32 GetChunkPaddingLength(uint64 Size);
 };
diff --git a/Engine/Plugins/Enterprise/GLTFExporter/Source/GLTFExporter/Public/Builders/GLTFJsonBuilder.h b/Engine/Plugins/Enterprise/GLTFExporter/Source/GLTFExporter/Public/Builders/GLTFJsonBuilder.h
--- a/Engine/Plugins/Enterprise/GLTFExporter/Source/GLTFExporter/Public/Builders/GLTFJsonBuilder.h
+++ b/Engine/Plugins/Enterprise/GLTFExporter/Source/GLTFExporter/Public/Builders/GLTFJsonBuilder.h
@@ -35,7 +35,7 @@ public:
 
 protected:
 
-	void WriteJsonArchive(FArchive& Archive);
+	bool WriteJsonArchive(FArchive& Archive);
 
 private:
 
