Add GLTFMeshSection::ToString

diff --git a/Engine/Plugins/Enterprise/GLTFExporter/Source/GLTFExporter/Private/Converters/GLTFMeshData.cpp b/Engine/Plugins/Enterprise/GLTFExporter/Source/GLTFExporter/Private/Converters/GLTFMeshData.cpp
--- a/Engine/Plugins/Enterprise/GLTFExporter/Source/GLTFExporter/Private/Converters/GLTFMeshData.cpp
+++ b/Engine/Plugins/Enterprise/GLTFExporter/Source/GLTFExporter/Private/Converters/GLTFMeshData.cpp
@@ -1,6 +1,7 @@
 // Copyright Epic Games, Inc. All Rights Reserved.
 
 #include "Converters/GLTFMeshData.h"
+#include "Converters/GLTFMeshUtilities.h"
 #include "Converters/GLTFNameUtilities.h"
 #include "Components/SkeletalMeshComponent.h"
 #include "Components/StaticMeshComponent.h"
@@ -61,7 +62,7 @@ FGLTFMeshData::FGLTFMeshData(const UStaticMesh* StaticMesh, const UStaticMeshCom
 
 #if WITH_EDITOR
 	LightMapTexCoord = StaticMesh->GetLightMapCoordinateIndex();
- 	const int32 NumTexCoords = StaticMesh->GetLODForExport(LODIndex).VertexBuffers.StaticMeshVertexBuffer.GetNumTexCoords();
+ 	const int32 NumTexCoords = FGLTFMeshUtilities::GetRenderData(StaticMesh, LODIndex).VertexBuffers.StaticMeshVertexBuffer.GetNumTexCoords();
 	BakeUsingTexCoord = FMath::Min(LightMapTexCoord, NumTexCoords - 1);
 	// TODO: add warning if texture coordinate for baking has overlap
 #endif
@@ -124,7 +125,7 @@ FGLTFMeshData::FGLTFMeshData(const USkeletalMesh* SkeletalMesh, const USkeletalM
 
 #if WITH_EDITOR
 	// TODO: don't assume last UV channel is non-overlapping
-	const int32 NumTexCoords = SkeletalMesh->GetResourceForRendering()->LODRenderData[LODIndex].StaticVertexBuffers.StaticMeshVertexBuffer.GetNumTexCoords();
+	const int32 NumTexCoords = FGLTFMeshUtilities::GetRenderData(SkeletalMesh, LODIndex).StaticVertexBuffers.StaticMeshVertexBuffer.GetNumTexCoords();
 	BakeUsingTexCoord = NumTexCoords - 1;
 	// TODO: add warning if texture coordinate for baking has overlap
 #endif
diff --git a/Engine/Plugins/Enterprise/GLTFExporter/Source/GLTFExporter/Private/Converters/GLTFMeshSection.cpp b/Engine/Plugins/Enterprise/GLTFExporter/Source/GLTFExporter/Private/Converters/GLTFMeshSection.cpp
--- a/Engine/Plugins/Enterprise/GLTFExporter/Source/GLTFExporter/Private/Converters/GLTFMeshSection.cpp
+++ b/Engine/Plugins/Enterprise/GLTFExporter/Source/GLTFExporter/Private/Converters/GLTFMeshSection.cpp
@@ -2,44 +2,55 @@
 
 #include "Converters/GLTFMeshSection.h"
 #include "Converters/GLTFBufferAdapter.h"
+#include "Converters/GLTFMeshUtilities.h"
+#include "Engine/StaticMesh.h"
+#include "Engine/SkeletalMesh.h"
 #include "Rendering/MultiSizeIndexContainer.h"
 #include "Rendering/SkeletalMeshRenderData.h"
 #include "Algo/MaxElement.h"
 
-FGLTFMeshSection::FGLTFMeshSection(const FStaticMeshLODResources* MeshLOD, const FGLTFIndexArray& SectionIndices)
-	: MaxBoneIndex(0)
+FGLTFMeshSection::FGLTFMeshSection(const UStaticMesh* StaticMesh, int32 LODIndex, const FGLTFIndexArray& SectionIndices)
+	: SectionIndices(SectionIndices)
+	, MaxBoneIndex(0)
 {
-	const TUniquePtr<IGLTFBufferAdapter> SourceBuffer = IGLTFBufferAdapter::GetIndices(&MeshLOD->IndexBuffer);
+	StaticMesh->GetName(Name);
+
+	const FStaticMeshLODResources& RenderData = FGLTFMeshUtilities::GetRenderData(StaticMesh, LODIndex);
+	const TUniquePtr<IGLTFBufferAdapter> SourceBuffer = IGLTFBufferAdapter::GetIndices(&RenderData.IndexBuffer);
 	const uint8* SourceData = SourceBuffer->GetData();
 
-	if (MeshLOD->IndexBuffer.Is32Bit())
+	if (RenderData.IndexBuffer.Is32Bit())
 	{
-		Init(MeshLOD->Sections, SectionIndices, reinterpret_cast<const uint32*>(SourceData));
+		Init(RenderData.Sections, reinterpret_cast<const uint32*>(SourceData));
 	}
 	else
 	{
-		Init(MeshLOD->Sections, SectionIndices, reinterpret_cast<const uint16*>(SourceData));
+		Init(RenderData.Sections, reinterpret_cast<const uint16*>(SourceData));
 	}
 }
 
-FGLTFMeshSection::FGLTFMeshSection(const FSkeletalMeshLODRenderData* MeshLOD, const FGLTFIndexArray& SectionIndices)
-	: MaxBoneIndex(0)
+FGLTFMeshSection::FGLTFMeshSection(const USkeletalMesh* SkeletalMesh, int32 LODIndex, const FGLTFIndexArray& SectionIndices)
+	: SectionIndices(SectionIndices)
+	, MaxBoneIndex(0)
 {
-	const TUniquePtr<IGLTFBufferAdapter> SourceBuffer = IGLTFBufferAdapter::GetIndices(MeshLOD->MultiSizeIndexContainer.GetIndexBuffer());
+	SkeletalMesh->GetName(Name);
+
+	const FSkeletalMeshLODRenderData& RenderData = FGLTFMeshUtilities::GetRenderData(SkeletalMesh, LODIndex);
+	const TUniquePtr<IGLTFBufferAdapter> SourceBuffer = IGLTFBufferAdapter::GetIndices(RenderData.MultiSizeIndexContainer.GetIndexBuffer());
 	const uint8* SourceData = SourceBuffer->GetData();
 
-	if (MeshLOD->MultiSizeIndexContainer.GetDataTypeSize() != sizeof(uint16))
+	if (RenderData.MultiSizeIndexContainer.GetDataTypeSize() != sizeof(uint16))
 	{
-		Init(MeshLOD->RenderSections, SectionIndices, reinterpret_cast<const uint32*>(SourceData));
+		Init(RenderData.RenderSections, reinterpret_cast<const uint32*>(SourceData));
 	}
 	else
 	{
-		Init(MeshLOD->RenderSections, SectionIndices, reinterpret_cast<const uint16*>(SourceData));
+		Init(RenderData.RenderSections, reinterpret_cast<const uint16*>(SourceData));
 	}
 }
 
 template <typename IndexType, typename SectionArrayType>
-void FGLTFMeshSection::Init(const SectionArrayType& Sections, const FGLTFIndexArray& SectionIndices, const IndexType* SourceData)
+void FGLTFMeshSection::Init(const SectionArrayType& Sections, const IndexType* SourceData)
 {
 	if (SourceData == nullptr)
 	{
@@ -104,6 +115,16 @@ void FGLTFMeshSection::Init(const SectionArrayType& Sections, const FGLTFIndexAr
 	}
 }
 
+FString FGLTFMeshSection::ToString() const
+{
+	return FString::Printf(
+		TEXT("%s (section%s %s)"),
+		*Name,
+		SectionIndices.Num() > 1 ? TEXT("s") : TEXT(""),
+		*FString::JoinBy(SectionIndices, TEXT(", "), FString::FromInt)
+		);
+}
+
 uint32 FGLTFMeshSection::GetIndexOffset(const FStaticMeshSection& Section)
 {
 	return Section.FirstIndex;
diff --git a/Engine/Plugins/Enterprise/GLTFExporter/Source/GLTFExporter/Private/Converters/GLTFMeshUtilities.cpp b/Engine/Plugins/Enterprise/GLTFExporter/Source/GLTFExporter/Private/Converters/GLTFMeshUtilities.cpp
--- a/Engine/Plugins/Enterprise/GLTFExporter/Source/GLTFExporter/Private/Converters/GLTFMeshUtilities.cpp
+++ b/Engine/Plugins/Enterprise/GLTFExporter/Source/GLTFExporter/Private/Converters/GLTFMeshUtilities.cpp
@@ -123,6 +123,16 @@ void FGLTFMeshUtilities::ResolveMaterials(TArray<const UMaterialInterface*>& Mat
 	}
 }
 
+const FStaticMeshLODResources& FGLTFMeshUtilities::GetRenderData(const UStaticMesh* StaticMesh, int32 LODIndex)
+{
+	return StaticMesh->GetLODForExport(LODIndex);
+}
+
+const FSkeletalMeshLODRenderData& FGLTFMeshUtilities::GetRenderData(const USkeletalMesh* SkeletalMesh, int32 LODIndex)
+{
+	return SkeletalMesh->GetResourceForRendering()->LODRenderData[LODIndex];
+}
+
 FGLTFIndexArray FGLTFMeshUtilities::GetSectionIndices(const UStaticMesh* StaticMesh, int32 LODIndex, int32 MaterialIndex)
 {
 	if (StaticMesh == nullptr)
@@ -130,8 +140,7 @@ FGLTFIndexArray FGLTFMeshUtilities::GetSectionIndices(const UStaticMesh* StaticM
 		return {};
 	}
 
-	const FStaticMeshLODResources& MeshLOD = StaticMesh->GetLODForExport(LODIndex);
-	return GetSectionIndices(MeshLOD, MaterialIndex);
+	return GetSectionIndices(GetRenderData(StaticMesh, LODIndex), MaterialIndex);
 }
 
 FGLTFIndexArray FGLTFMeshUtilities::GetSectionIndices(const USkeletalMesh* SkeletalMesh, int32 LODIndex, int32 MaterialIndex)
@@ -141,14 +150,12 @@ FGLTFIndexArray FGLTFMeshUtilities::GetSectionIndices(const USkeletalMesh* Skele
 		return {};
 	}
 
-	const FSkeletalMeshRenderData* RenderData = SkeletalMesh->GetResourceForRendering();
-	const FSkeletalMeshLODRenderData& MeshLOD = RenderData->LODRenderData[LODIndex];
-	return GetSectionIndices(MeshLOD, MaterialIndex);
+	return GetSectionIndices(GetRenderData(SkeletalMesh, LODIndex), MaterialIndex);
 }
 
-FGLTFIndexArray FGLTFMeshUtilities::GetSectionIndices(const FStaticMeshLODResources& MeshLOD, int32 MaterialIndex)
+FGLTFIndexArray FGLTFMeshUtilities::GetSectionIndices(const FStaticMeshLODResources& RenderData, int32 MaterialIndex)
 {
-	const FStaticMeshSectionArray& Sections = MeshLOD.Sections;
+	const FStaticMeshSectionArray& Sections = RenderData.Sections;
 
 	FGLTFIndexArray SectionIndices;
 	SectionIndices.Reserve(Sections.Num());
@@ -164,9 +171,9 @@ FGLTFIndexArray FGLTFMeshUtilities::GetSectionIndices(const FStaticMeshLODResour
 	return SectionIndices;
 }
 
-FGLTFIndexArray FGLTFMeshUtilities::GetSectionIndices(const FSkeletalMeshLODRenderData& MeshLOD, int32 MaterialIndex)
+FGLTFIndexArray FGLTFMeshUtilities::GetSectionIndices(const FSkeletalMeshLODRenderData& RenderData, int32 MaterialIndex)
 {
-	const TArray<FSkelMeshRenderSection>& Sections = MeshLOD.RenderSections;
+	const TArray<FSkelMeshRenderSection>& Sections = RenderData.RenderSections;
 
 	FGLTFIndexArray SectionIndices;
 	SectionIndices.Reserve(Sections.Num());
diff --git a/Engine/Plugins/Enterprise/GLTFExporter/Source/GLTFExporter/Private/Converters/GLTFMeshUtilities.h b/Engine/Plugins/Enterprise/GLTFExporter/Source/GLTFExporter/Private/Converters/GLTFMeshUtilities.h
--- a/Engine/Plugins/Enterprise/GLTFExporter/Source/GLTFExporter/Private/Converters/GLTFMeshUtilities.h
+++ b/Engine/Plugins/Enterprise/GLTFExporter/Source/GLTFExporter/Private/Converters/GLTFMeshUtilities.h
@@ -35,11 +35,14 @@ struct FGLTFMeshUtilities
 	static void ResolveMaterials(TArray<const UMaterialInterface*>& Materials, const TArray<MaterialType>& Defaults);
 	static void ResolveMaterials(TArray<const UMaterialInterface*>& Materials, const UMaterialInterface* Default);
 
+	static const FStaticMeshLODResources& GetRenderData(const UStaticMesh* StaticMesh, int32 LODIndex);
+	static const FSkeletalMeshLODRenderData& GetRenderData(const USkeletalMesh* SkeletalMesh, int32 LODIndex);
+
 	static FGLTFIndexArray GetSectionIndices(const UStaticMesh* StaticMesh, int32 LODIndex, int32 MaterialIndex);
 	static FGLTFIndexArray GetSectionIndices(const USkeletalMesh* SkeletalMesh, int32 LODIndex, int32 MaterialIndex);
 
-	static FGLTFIndexArray GetSectionIndices(const FStaticMeshLODResources& MeshLOD, int32 MaterialIndex);
-	static FGLTFIndexArray GetSectionIndices(const FSkeletalMeshLODRenderData& MeshLOD, int32 MaterialIndex);
+	static FGLTFIndexArray GetSectionIndices(const FStaticMeshLODResources& RenderData, int32 MaterialIndex);
+	static FGLTFIndexArray GetSectionIndices(const FSkeletalMeshLODRenderData& RenderData, int32 MaterialIndex);
 
 	static int32 GetLOD(const UStaticMesh* StaticMesh, const UStaticMeshComponent* StaticMeshComponent, int32 DefaultLOD);
 	static int32 GetLOD(const USkeletalMesh* SkeletalMesh, const USkeletalMeshComponent* SkeletalMeshComponent, int32 DefaultLOD);
diff --git a/Engine/Plugins/Enterprise/GLTFExporter/Source/GLTFExporter/Private/Tasks/GLTFDelayedMeshTasks.cpp b/Engine/Plugins/Enterprise/GLTFExporter/Source/GLTFExporter/Private/Tasks/GLTFDelayedMeshTasks.cpp
--- a/Engine/Plugins/Enterprise/GLTFExporter/Source/GLTFExporter/Private/Tasks/GLTFDelayedMeshTasks.cpp
+++ b/Engine/Plugins/Enterprise/GLTFExporter/Source/GLTFExporter/Private/Tasks/GLTFDelayedMeshTasks.cpp
@@ -119,10 +119,10 @@ void FGLTFDelayedStaticMeshTask::Process()
 	FGLTFMeshUtilities::FullyLoad(StaticMesh);
 	JsonMesh->Name = StaticMeshComponent != nullptr ? FGLTFNameUtilities::GetName(StaticMeshComponent) : StaticMesh->GetName();
 
-	const FStaticMeshLODResources& MeshLOD = StaticMesh->GetLODForExport(LODIndex);
-	const FPositionVertexBuffer& PositionBuffer = MeshLOD.VertexBuffers.PositionVertexBuffer;
-	const FStaticMeshVertexBuffer& VertexBuffer = MeshLOD.VertexBuffers.StaticMeshVertexBuffer;
-	const FColorVertexBuffer* ColorBuffer = &MeshLOD.VertexBuffers.ColorVertexBuffer; // TODO: add support for overriding color buffer by component
+	const FStaticMeshLODResources& RenderData = FGLTFMeshUtilities::GetRenderData(StaticMesh, LODIndex);
+	const FPositionVertexBuffer& PositionBuffer = RenderData.VertexBuffers.PositionVertexBuffer;
+	const FStaticMeshVertexBuffer& VertexBuffer = RenderData.VertexBuffers.StaticMeshVertexBuffer;
+	const FColorVertexBuffer* ColorBuffer = &RenderData.VertexBuffers.ColorVertexBuffer; // TODO: add support for overriding color buffer by component
 
 	if (Builder.ExportOptions->bExportVertexColors && HasVertexColors(ColorBuffer))
 	{
@@ -167,8 +167,8 @@ void FGLTFDelayedStaticMeshTask::Process()
 
 	for (int32 MaterialIndex = 0; MaterialIndex < MaterialCount; ++MaterialIndex)
 	{
-		const FGLTFIndexArray SectionIndices = FGLTFMeshUtilities::GetSectionIndices(MeshLOD, MaterialIndex);
-		const FGLTFMeshSection* ConvertedSection = MeshSectionConverter.GetOrAdd(&MeshLOD, SectionIndices);
+		const FGLTFIndexArray SectionIndices = FGLTFMeshUtilities::GetSectionIndices(RenderData, MaterialIndex);
+		const FGLTFMeshSection* ConvertedSection = MeshSectionConverter.GetOrAdd(StaticMesh, LODIndex, SectionIndices);
 
 		FGLTFJsonPrimitive& JsonPrimitive = JsonMesh->Primitives[MaterialIndex];
 		JsonPrimitive.Indices = Builder.AddUniqueIndexAccessor(ConvertedSection);
@@ -176,16 +176,11 @@ void FGLTFDelayedStaticMeshTask::Process()
 		JsonPrimitive.Attributes.Position = Builder.AddUniquePositionAccessor(ConvertedSection, &PositionBuffer);
 		if (JsonPrimitive.Attributes.Position == nullptr)
 		{
-			FString SectionString = TEXT("section");
-			SectionString += SectionIndices.Num() > 1 ? TEXT("s ") : TEXT(" ");
-			SectionString += FString::JoinBy(SectionIndices, TEXT(", "), FString::FromInt);
-
 			Builder.LogError(
-				FString::Printf(TEXT("Failed to export vertex positions related to material slot %d (%s) in static mesh %s (sections %s)"),
+				FString::Printf(TEXT("Failed to export vertex positions related to material slot %d (%s) in static mesh %s"),
 				MaterialIndex,
 				*MaterialSlots[MaterialIndex].MaterialSlotName.ToString(),
-				*StaticMesh->GetName(),
-				*SectionString
+				*ConvertedSection->ToString()
 				));
 		}
 
@@ -222,13 +217,11 @@ void FGLTFDelayedSkeletalMeshTask::Process()
 	FGLTFMeshUtilities::FullyLoad(SkeletalMesh);
 	JsonMesh->Name = SkeletalMeshComponent != nullptr ? FGLTFNameUtilities::GetName(SkeletalMeshComponent) : SkeletalMesh->GetName();
 
-	const FSkeletalMeshRenderData* RenderData = SkeletalMesh->GetResourceForRendering();
-	const FSkeletalMeshLODRenderData& MeshLOD = RenderData->LODRenderData[LODIndex];
-
-	const FPositionVertexBuffer& PositionBuffer = MeshLOD.StaticVertexBuffers.PositionVertexBuffer;
-	const FStaticMeshVertexBuffer& VertexBuffer = MeshLOD.StaticVertexBuffers.StaticMeshVertexBuffer;
-	const FColorVertexBuffer* ColorBuffer = &MeshLOD.StaticVertexBuffers.ColorVertexBuffer; // TODO: add support for overriding color buffer by component
-	const FSkinWeightVertexBuffer* SkinWeightBuffer = MeshLOD.GetSkinWeightVertexBuffer(); // TODO: add support for overriding skin weight buffer by component
+	const FSkeletalMeshLODRenderData& RenderData = FGLTFMeshUtilities::GetRenderData(SkeletalMesh, LODIndex);
+	const FPositionVertexBuffer& PositionBuffer = RenderData.StaticVertexBuffers.PositionVertexBuffer;
+	const FStaticMeshVertexBuffer& VertexBuffer = RenderData.StaticVertexBuffers.StaticMeshVertexBuffer;
+	const FColorVertexBuffer* ColorBuffer = &RenderData.StaticVertexBuffers.ColorVertexBuffer; // TODO: add support for overriding color buffer by component
+	const FSkinWeightVertexBuffer* SkinWeightBuffer = RenderData.GetSkinWeightVertexBuffer(); // TODO: add support for overriding skin weight buffer by component
 	// TODO: add support for skin weight profiles?
 	// TODO: add support for morph targets
 
@@ -276,8 +269,8 @@ void FGLTFDelayedSkeletalMeshTask::Process()
 
 	for (int32 MaterialIndex = 0; MaterialIndex < MaterialCount; ++MaterialIndex)
 	{
-		const FGLTFIndexArray SectionIndices = FGLTFMeshUtilities::GetSectionIndices(MeshLOD, MaterialIndex);
-		const FGLTFMeshSection* ConvertedSection = MeshSectionConverter.GetOrAdd(&MeshLOD, SectionIndices);
+		const FGLTFIndexArray SectionIndices = FGLTFMeshUtilities::GetSectionIndices(RenderData, MaterialIndex);
+		const FGLTFMeshSection* ConvertedSection = MeshSectionConverter.GetOrAdd(SkeletalMesh, LODIndex, SectionIndices);
 
 		FGLTFJsonPrimitive& JsonPrimitive = JsonMesh->Primitives[MaterialIndex];
 		JsonPrimitive.Indices = Builder.AddUniqueIndexAccessor(ConvertedSection);
@@ -285,16 +278,11 @@ void FGLTFDelayedSkeletalMeshTask::Process()
 		JsonPrimitive.Attributes.Position = Builder.AddUniquePositionAccessor(ConvertedSection, &PositionBuffer);
 		if (JsonPrimitive.Attributes.Position == nullptr)
 		{
-			FString SectionString = TEXT("section");
-			SectionString += SectionIndices.Num() > 1 ? TEXT("s ") : TEXT(" ");
-			SectionString += FString::JoinBy(SectionIndices, TEXT(", "), FString::FromInt);
-
 			Builder.LogError(
-				FString::Printf(TEXT("Failed to export vertex positions related to material slot %d (%s) in skeletal mesh %s (sections %s)"),
+				FString::Printf(TEXT("Failed to export vertex positions related to material slot %d (%s) in skeletal mesh %s"),
 				MaterialIndex,
 				*MaterialSlots[MaterialIndex].MaterialSlotName.ToString(),
-				*SkeletalMesh->GetName(),
-				*SectionString
+				*ConvertedSection->ToString()
 				));
 		}
 
diff --git a/Engine/Plugins/Enterprise/GLTFExporter/Source/GLTFExporter/Public/Converters/GLTFMeshSection.h b/Engine/Plugins/Enterprise/GLTFExporter/Source/GLTFExporter/Public/Converters/GLTFMeshSection.h
--- a/Engine/Plugins/Enterprise/GLTFExporter/Source/GLTFExporter/Public/Converters/GLTFMeshSection.h
+++ b/Engine/Plugins/Enterprise/GLTFExporter/Source/GLTFExporter/Public/Converters/GLTFMeshSection.h
@@ -6,15 +6,18 @@
 #include "BoneIndices.h"
 #include "Converters/GLTFIndexArray.h"
 
-class FSkeletalMeshLODRenderData;
+class UStaticMesh;
+class USkeletalMesh;
 struct FSkelMeshRenderSection;
-struct FStaticMeshLODResources;
 struct FStaticMeshSection;
 
 struct GLTFEXPORTER_API FGLTFMeshSection
 {
-	FGLTFMeshSection(const FStaticMeshLODResources* MeshLOD, const FGLTFIndexArray& SectionIndices);
-	FGLTFMeshSection(const FSkeletalMeshLODRenderData* MeshLOD, const FGLTFIndexArray& SectionIndices);
+	FGLTFMeshSection(const UStaticMesh* Mesh, int32 LODIndex, const FGLTFIndexArray& SectionIndices);
+	FGLTFMeshSection(const USkeletalMesh* Mesh, int32 LODIndex, const FGLTFIndexArray& SectionIndices);
+
+	FString Name;
+	FGLTFIndexArray SectionIndices;
 
 	TArray<uint32> IndexMap;
 	TArray<uint32> IndexBuffer;
@@ -23,10 +26,12 @@ struct GLTFEXPORTER_API FGLTFMeshSection
 	TArray<uint32> BoneMapLookup;
 	FBoneIndexType MaxBoneIndex;
 
+	FString ToString() const;
+
 private:
 
 	template <typename IndexType, typename SectionArrayType>
-	void Init(const SectionArrayType& Sections, const FGLTFIndexArray& SectionIndices, const IndexType* SourceData);
+	void Init(const SectionArrayType& Sections, const IndexType* SourceData);
 
 	static uint32 GetIndexOffset(const FStaticMeshSection& Section);
 	static uint32 GetIndexOffset(const FSkelMeshRenderSection& Section);
diff --git a/Engine/Plugins/Enterprise/GLTFExporter/Source/GLTFExporter/Public/Converters/GLTFMeshSectionConverters.h b/Engine/Plugins/Enterprise/GLTFExporter/Source/GLTFExporter/Public/Converters/GLTFMeshSectionConverters.h
--- a/Engine/Plugins/Enterprise/GLTFExporter/Source/GLTFExporter/Public/Converters/GLTFMeshSectionConverters.h
+++ b/Engine/Plugins/Enterprise/GLTFExporter/Source/GLTFExporter/Public/Converters/GLTFMeshSectionConverters.h
@@ -6,8 +6,8 @@
 #include "Converters/GLTFMeshSection.h"
 #include "Converters/GLTFIndexArray.h"
 
-template <typename MeshLODType>
-class TGLTFMeshSectionConverter : public TGLTFConverter<const FGLTFMeshSection*, const MeshLODType*, FGLTFIndexArray>
+template <typename MeshType>
+class TGLTFMeshSectionConverter : public TGLTFConverter<const FGLTFMeshSection*, const MeshType*, int32, FGLTFIndexArray>
 {
 public:
 
@@ -20,9 +20,9 @@ public:
 
 protected:
 
-	virtual const FGLTFMeshSection* Convert(const MeshLODType* MeshLOD, FGLTFIndexArray SectionIndices) override
+	virtual const FGLTFMeshSection* Convert(const MeshType* Mesh, int32 LODIndex, FGLTFIndexArray SectionIndices) override
 	{
-		return Outputs.Add_GetRef(MakeUnique<FGLTFMeshSection>(MeshLOD, SectionIndices)).Get();
+		return Outputs.Add_GetRef(MakeUnique<FGLTFMeshSection>(Mesh, LODIndex, SectionIndices)).Get();
 	}
 
 private:
@@ -30,5 +30,5 @@ private:
 	TArray<TUniquePtr<FGLTFMeshSection>> Outputs;
 };
 
-typedef TGLTFMeshSectionConverter<FStaticMeshLODResources> FGLTFStaticMeshSectionConverter;
-typedef TGLTFMeshSectionConverter<FSkeletalMeshLODRenderData> FGLTFSkeletalMeshSectionConverter;
+typedef TGLTFMeshSectionConverter<UStaticMesh> FGLTFStaticMeshSectionConverter;
+typedef TGLTFMeshSectionConverter<USkeletalMesh> FGLTFSkeletalMeshSectionConverter;
diff --git a/Engine/Plugins/Enterprise/GLTFExporter/Source/GLTFExporter/Public/Utilities/GLTFCoreUtilities.h b/Engine/Plugins/Enterprise/GLTFExporter/Source/GLTFExporter/Public/Utilities/GLTFCoreUtilities.h
--- a/Engine/Plugins/Enterprise/GLTFExporter/Source/GLTFExporter/Public/Utilities/GLTFCoreUtilities.h
+++ b/Engine/Plugins/Enterprise/GLTFExporter/Source/GLTFExporter/Public/Utilities/GLTFCoreUtilities.h
@@ -71,13 +71,13 @@ struct GLTFEXPORTER_API FGLTFCoreUtilities
 	template <typename ComponentType>
 	static EGLTFJsonComponentType GetComponentType()
 	{
-		     if constexpr (std::is_same_v<ComponentType, int8  >) return EGLTFJsonComponentType::Int8;
-		else if constexpr (std::is_same_v<ComponentType, uint8 >) return EGLTFJsonComponentType::UInt8;
-		else if constexpr (std::is_same_v<ComponentType, int16 >) return EGLTFJsonComponentType::Int16;
-		else if constexpr (std::is_same_v<ComponentType, uint16>) return EGLTFJsonComponentType::UInt16;
-		else if constexpr (std::is_same_v<ComponentType, int32 >) return EGLTFJsonComponentType::Int32;
-		else if constexpr (std::is_same_v<ComponentType, uint32>) return EGLTFJsonComponentType::UInt32;
-		else if constexpr (std::is_same_v<ComponentType, float >) return EGLTFJsonComponentType::Float;
-		else return EGLTFJsonComponentType::None;
+		if constexpr (std::is_same_v<ComponentType, int8  >) return EGLTFJsonComponentType::Int8;
+		if constexpr (std::is_same_v<ComponentType, uint8 >) return EGLTFJsonComponentType::UInt8;
+		if constexpr (std::is_same_v<ComponentType, int16 >) return EGLTFJsonComponentType::Int16;
+		if constexpr (std::is_same_v<ComponentType, uint16>) return EGLTFJsonComponentType::UInt16;
+		if constexpr (std::is_same_v<ComponentType, int32 >) return EGLTFJsonComponentType::Int32;
+		if constexpr (std::is_same_v<ComponentType, uint32>) return EGLTFJsonComponentType::UInt32;
+		if constexpr (std::is_same_v<ComponentType, float >) return EGLTFJsonComponentType::Float;
+		return EGLTFJsonComponentType::None;
 	}
 };
