Fix maximum values used in exported glTF mesh.primitive.indices

diff --git a/Engine/Plugins/Enterprise/GLTFExporter/Source/GLTFExporter/Private/Converters/GLTFAccessorConverters.cpp b/Engine/Plugins/Enterprise/GLTFExporter/Source/GLTFExporter/Private/Converters/GLTFAccessorConverters.cpp
--- a/Engine/Plugins/Enterprise/GLTFExporter/Source/GLTFExporter/Private/Converters/GLTFAccessorConverters.cpp
+++ b/Engine/Plugins/Enterprise/GLTFExporter/Source/GLTFExporter/Private/Converters/GLTFAccessorConverters.cpp
@@ -571,8 +571,18 @@ FGLTFJsonAccessor* FGLTFBoneWeightBufferConverter::Convert(const FGLTFMeshSectio
 FGLTFJsonAccessor* FGLTFIndexBufferConverter::Convert(const FGLTFMeshSection* MeshSection)
 {
 	const uint32 MaxVertexIndex = MeshSection->IndexMap.Num() - 1;
-	if (MaxVertexIndex <= UINT8_MAX) return Convert<uint8>(MeshSection);
-	if (MaxVertexIndex <= UINT16_MAX) return Convert<uint16>(MeshSection);
+	// NOTE: Even if a maximum value (i.e. UINT8_MAX or UINT16_MAX) would fit inside the type (i.e. uint8 or uint16), it's not allowed per gltf spec
+	// since the maximum values trigger primitive restart in some graphics APIs and would require client implementations to rebuild the index buffer.
+	if (MaxVertexIndex < UINT8_MAX) return Convert<uint8>(MeshSection);
+	if (MaxVertexIndex < UINT16_MAX) return Convert<uint16>(MeshSection);
+	if (MaxVertexIndex == UINT32_MAX)
+	{
+		Builder.LogError(
+			FString::Printf(TEXT("Maximum value (%d) used in indices accessor for mesh %s, this may not be supported on some glTF viewers"),
+			UINT32_MAX,
+			*MeshSection->ToString()
+			));
+	}
 	return Convert<uint32>(MeshSection);
 }
 
