Minor code cleanup in GLTFMaterialUtilities

diff --git a/Engine/Plugins/Enterprise/GLTFExporter/Source/GLTFExporter/Private/Converters/GLTFMaterialUtilities.cpp b/Engine/Plugins/Enterprise/GLTFExporter/Source/GLTFExporter/Private/Converters/GLTFMaterialUtilities.cpp
--- a/Engine/Plugins/Enterprise/GLTFExporter/Source/GLTFExporter/Private/Converters/GLTFMaterialUtilities.cpp
+++ b/Engine/Plugins/Enterprise/GLTFExporter/Source/GLTFExporter/Private/Converters/GLTFMaterialUtilities.cpp
@@ -164,19 +164,11 @@ FGLTFPropertyBakeOutput FGLTFMaterialUtilities::BakeMaterialProperty(const FIntP
 	}
 
 	bool bFromSRGB = !bIsLinearBake;
+	FGLTFPropertyBakeOutput PropertyBakeOutput(BakedPixels, BakedSize, EmissiveScale, bFromSRGB);
+
 	bool bToSRGB = IsSRGB(Property);
 	FGLTFTextureUtilities::TransformColorSpace(*BakedPixels, bFromSRGB, bToSRGB);
 
-	FGLTFPropertyBakeOutput PropertyBakeOutput(Property, PF_B8G8R8A8, BakedPixels, BakedSize, EmissiveScale, !bIsLinearBake);
-
-	if (BakedPixels->Num() == 1)
-	{
-		const FColor& Pixel = (*BakedPixels)[0];
-
-		PropertyBakeOutput.bIsConstant = true;
-		PropertyBakeOutput.ConstantValue = bToSRGB ? FLinearColor(Pixel) : Pixel.ReinterpretAsLinear();
-	}
-
 	return PropertyBakeOutput;
 }
 
diff --git a/Engine/Plugins/Enterprise/GLTFExporter/Source/GLTFExporter/Private/Converters/GLTFMaterialUtilities.h b/Engine/Plugins/Enterprise/GLTFExporter/Source/GLTFExporter/Private/Converters/GLTFMaterialUtilities.h
--- a/Engine/Plugins/Enterprise/GLTFExporter/Source/GLTFExporter/Private/Converters/GLTFMaterialUtilities.h
+++ b/Engine/Plugins/Enterprise/GLTFExporter/Source/GLTFExporter/Private/Converters/GLTFMaterialUtilities.h
@@ -15,22 +15,22 @@ class UMaterialExpressionCustomOutput;
 
 struct FGLTFPropertyBakeOutput
 {
-	FGLTFPropertyBakeOutput(const FMaterialPropertyEx& Property, EPixelFormat PixelFormat, const TGLTFSharedArray<FColor>& Pixels, FIntPoint Size, float EmissiveScale, bool bIsSRGB)
-		: Property(Property)
-		, PixelFormat(PixelFormat)
-		, Pixels(Pixels)
+	FGLTFPropertyBakeOutput(const TGLTFSharedArray<FColor>& Pixels, FIntPoint Size, float EmissiveScale, bool bSRGB)
+		: Pixels(Pixels)
 		, Size(Size)
 		, EmissiveScale(EmissiveScale)
-		, bIsSRGB(bIsSRGB)
-		, bIsConstant(false)
-	{}
+		, bIsConstant(Pixels->Num() == 1)
+	{
+		if (bIsConstant)
+		{
+			const FColor& Pixel = (*Pixels)[0];
+			ConstantValue = bSRGB ? FLinearColor(Pixel) : Pixel.ReinterpretAsLinear();
+		}
+	}
 
-	const FMaterialPropertyEx& Property;
-	EPixelFormat PixelFormat;
 	TGLTFSharedArray<FColor> Pixels;
 	FIntPoint Size;
 	float EmissiveScale;
-	bool bIsSRGB;
 	bool bIsConstant;
 	FLinearColor ConstantValue;
 };
diff --git a/Engine/Plugins/Enterprise/GLTFExporter/Source/GLTFExporter/Private/Tasks/GLTFDelayedMaterialTasks.cpp b/Engine/Plugins/Enterprise/GLTFExporter/Source/GLTFExporter/Private/Tasks/GLTFDelayedMaterialTasks.cpp
--- a/Engine/Plugins/Enterprise/GLTFExporter/Source/GLTFExporter/Private/Tasks/GLTFDelayedMaterialTasks.cpp
+++ b/Engine/Plugins/Enterprise/GLTFExporter/Source/GLTFExporter/Private/Tasks/GLTFDelayedMaterialTasks.cpp
@@ -885,7 +885,7 @@ bool FGLTFDelayedMaterialTask::TryGetEmissive(FGLTFJsonMaterial& OutMaterial, co
 				return true;
 			}
 
-			if (!StoreBakedPropertyTexture(OutMaterial.EmissiveTexture, PropertyBakeOutput, TEXT("Emissive")))
+			if (!StoreBakedPropertyTexture(EmissiveProperty, OutMaterial.EmissiveTexture, PropertyBakeOutput, TEXT("Emissive")))
 			{
 				return false;
 			}
@@ -1345,7 +1345,7 @@ bool FGLTFDelayedMaterialTask::TryGetBakedMaterialProperty(FGLTFJsonTextureInfo&
 		return true;
 	}
 
-	if (StoreBakedPropertyTexture(OutTexInfo, PropertyBakeOutput, PropertyName))
+	if (StoreBakedPropertyTexture(Property, OutTexInfo, PropertyBakeOutput, PropertyName))
 	{
 		OutConstant = FGLTFJsonColor3::White; // make sure property is not zero
 		return true;
@@ -1379,7 +1379,7 @@ bool FGLTFDelayedMaterialTask::TryGetBakedMaterialProperty(FGLTFJsonTextureInfo&
 		return true;
 	}
 
-	if (StoreBakedPropertyTexture(OutTexInfo, PropertyBakeOutput, PropertyName))
+	if (StoreBakedPropertyTexture(Property, OutTexInfo, PropertyBakeOutput, PropertyName))
 	{
 		OutConstant = FGLTFJsonColor4::White; // make sure property is not zero
 		return true;
@@ -1413,7 +1413,7 @@ inline bool FGLTFDelayedMaterialTask::TryGetBakedMaterialProperty(FGLTFJsonTextu
 		return true;
 	}
 
-	if (StoreBakedPropertyTexture(OutTexInfo, PropertyBakeOutput, PropertyName))
+	if (StoreBakedPropertyTexture(Property, OutTexInfo, PropertyBakeOutput, PropertyName))
 	{
 		OutConstant = 1; // make sure property is not zero
 		return true;
@@ -1443,7 +1443,7 @@ bool FGLTFDelayedMaterialTask::TryGetBakedMaterialProperty(FGLTFJsonTextureInfo&
 			return true;
 		}
 
-		return StoreBakedPropertyTexture(OutTexInfo, PropertyBakeOutput, PropertyName);
+		return StoreBakedPropertyTexture(Property, OutTexInfo, PropertyBakeOutput, PropertyName);
 	}
 
 	const FVector4f MaskedConstant = FVector4f(PropertyBakeOutput.ConstantValue) * FGLTFMaterialUtilities::GetPropertyMask(Property);
@@ -1556,17 +1556,18 @@ FGLTFPropertyBakeOutput FGLTFDelayedMaterialTask::BakeMaterialProperty(const FMa
 		Builder.ExportOptions->bAdjustNormalmaps);
 }
 
-bool FGLTFDelayedMaterialTask::StoreBakedPropertyTexture(FGLTFJsonTextureInfo& OutTexInfo, FGLTFPropertyBakeOutput& PropertyBakeOutput, const FString& PropertyName) const
+bool FGLTFDelayedMaterialTask::StoreBakedPropertyTexture(const FMaterialPropertyEx& Property, FGLTFJsonTextureInfo& OutTexInfo, FGLTFPropertyBakeOutput& PropertyBakeOutput, const FString& PropertyName) const
 {
-	const TextureAddress TextureAddress = Builder.GetBakeTilingForMaterialProperty(Material, GetPropertyGroup(PropertyBakeOutput.Property));
-	const TextureFilter TextureFilter = Builder.GetBakeFilterForMaterialProperty(Material, GetPropertyGroup(PropertyBakeOutput.Property));
+	const EGLTFMaterialPropertyGroup PropertyGroup = GetPropertyGroup(Property);
+	const TextureAddress TextureAddress = Builder.GetBakeTilingForMaterialProperty(Material, PropertyGroup);
+	const TextureFilter TextureFilter = Builder.GetBakeFilterForMaterialProperty(Material, PropertyGroup);
 
 	FGLTFJsonTexture* Texture = FGLTFMaterialUtilities::AddTexture(
 		Builder,
 		PropertyBakeOutput.Pixels,
 		PropertyBakeOutput.Size,
 		true, // NOTE: we can ignore alpha in everything but TryGetBaseColorAndOpacity
-		FGLTFMaterialUtilities::IsNormalMap(PropertyBakeOutput.Property),
+		FGLTFMaterialUtilities::IsNormalMap(Property),
 		GetBakedTextureName(PropertyName),
 		TextureAddress,
 		TextureFilter);
diff --git a/Engine/Plugins/Enterprise/GLTFExporter/Source/GLTFExporter/Private/Tasks/GLTFDelayedMaterialTasks.h b/Engine/Plugins/Enterprise/GLTFExporter/Source/GLTFExporter/Private/Tasks/GLTFDelayedMaterialTasks.h
--- a/Engine/Plugins/Enterprise/GLTFExporter/Source/GLTFExporter/Private/Tasks/GLTFDelayedMaterialTasks.h
+++ b/Engine/Plugins/Enterprise/GLTFExporter/Source/GLTFExporter/Private/Tasks/GLTFDelayedMaterialTasks.h
@@ -93,7 +93,7 @@ private:
 	FGLTFPropertyBakeOutput BakeMaterialProperty(const FMaterialPropertyEx& Property, int32& OutTexCoord, FGLTFJsonTextureTransform& OutTransform);
 	FGLTFPropertyBakeOutput BakeMaterialProperty(const FMaterialPropertyEx& Property, int32& OutTexCoord, FGLTFJsonTextureTransform& OutTransform, const FIntPoint& TextureSize, bool bFillAlpha);
 
-	bool StoreBakedPropertyTexture(FGLTFJsonTextureInfo& OutTexInfo, FGLTFPropertyBakeOutput& PropertyBakeOutput, const FString& PropertyName) const;
+	bool StoreBakedPropertyTexture(const FMaterialPropertyEx& Property, FGLTFJsonTextureInfo& OutTexInfo, FGLTFPropertyBakeOutput& PropertyBakeOutput, const FString& PropertyName) const;
 
 	static EGLTFMaterialPropertyGroup GetPropertyGroup(const FMaterialPropertyEx& Property);
 
