Extend glTF exporter analytics to also report asset complexity

diff --git a/Engine/Plugins/Enterprise/GLTFExporter/Source/GLTFExporter/Private/Builders/GLTFAnalyticsBuilder.cpp b/Engine/Plugins/Enterprise/GLTFExporter/Source/GLTFExporter/Private/Builders/GLTFAnalyticsBuilder.cpp
new file mode 100644
--- /dev/null
+++ b/Engine/Plugins/Enterprise/GLTFExporter/Source/GLTFExporter/Private/Builders/GLTFAnalyticsBuilder.cpp
@@ -0,0 +1,106 @@
+// Copyright Epic Games, Inc. All Rights Reserved.
+
+#include "Builders/GLTFAnalyticsBuilder.h"
+#include "EngineAnalytics.h"
+#include "StaticMeshResources.h"
+#include "Rendering/SkeletalMeshLODRenderData.h"
+#include "Converters/GLTFMeshUtilities.h"
+
+FGLTFAnalyticsBuilder::FGLTFAnalyticsBuilder(const FString& FileName, const UGLTFExportOptions* ExportOptions)
+	: FGLTFBufferBuilder(FileName, ExportOptions)
+{
+}
+
+TArray<FAnalyticsEventAttribute> FGLTFAnalyticsBuilder::GenerateAnalytics() const
+{
+	int32 NumberOfActors = ActorsRecorded.Num();
+	int32 NumberOfStaticMeshes = StaticMeshesRecorded.Num();
+	int32 NumberOfSkeletalMeshes = SkeletalMeshesRecorded.Num();
+	int32 NumberOfLevelSequences = LevelSequencesRecorded.Num();
+	int32 NumberOfAnimSequences = AnimSequencesRecorded.Num();
+	int32 NumberOfMaterials = MaterialsRecorded.Num();
+	int32 NumberOfTextures = TexturesRecorded.Num();
+	int32 NumberOfCameras = CamerasRecorded.Num();
+	int32 NumberOfLights = LightsRecorded.Num();
+
+	TArray<FAnalyticsEventAttribute> EventAttributes;
+	auto Add = [&EventAttributes](const FString& Name, int32 Value)
+	{
+		EventAttributes.Emplace(Name, LexToString(Value));
+	};
+
+	Add(TEXT("NumberOfActors"), NumberOfActors);
+	Add(TEXT("NumberOfStaticMeshes"), NumberOfStaticMeshes);
+	Add(TEXT("NumberOfSkeletalMeshes"), NumberOfSkeletalMeshes);
+	Add(TEXT("NumberOfLevelSequences"), NumberOfLevelSequences);
+	Add(TEXT("NumberOfAnimSequences"), NumberOfAnimSequences);
+	Add(TEXT("NumberOfMaterials"), NumberOfMaterials);
+	Add(TEXT("NumberOfTextures"), NumberOfTextures);
+	Add(TEXT("NumberOfCameras"), NumberOfCameras);
+	Add(TEXT("NumberOfLights"), NumberOfLights);
+
+	uint64 NumberOfVertices = 0;
+	{
+		//Statics:
+		for (const UStaticMesh* StaticMesh : StaticMeshesRecorded)
+		{
+			const FStaticMeshLODResources& RenderData = FGLTFMeshUtilities::GetRenderData(StaticMesh, 0);
+			NumberOfVertices += RenderData.GetNumVertices();
+		}
+
+		//Skeletals:
+		for (const USkeletalMesh* SkeletalMesh : SkeletalMeshesRecorded)
+		{
+			const FSkeletalMeshLODRenderData& RenderData = FGLTFMeshUtilities::GetRenderData(SkeletalMesh, 0);
+			NumberOfVertices += RenderData.GetNumVertices();
+		}
+	}
+	EventAttributes.Emplace(TEXT("NumberOfVertices"), LexToString(NumberOfVertices));
+
+	return EventAttributes;
+}
+
+void FGLTFAnalyticsBuilder::RecordActor(const AActor* Object)
+{
+	ActorsRecorded.Add(Object);
+}
+
+void FGLTFAnalyticsBuilder::RecordStaticMesh(const UStaticMesh* Object)
+{
+	StaticMeshesRecorded.Add(Object);
+}
+
+void FGLTFAnalyticsBuilder::RecordSkeletalMesh(const USkeletalMesh* Object)
+{
+	SkeletalMeshesRecorded.Add(Object);
+}
+
+void FGLTFAnalyticsBuilder::RecordLevelSequence(const ULevelSequence* Object)
+{
+	LevelSequencesRecorded.Add(Object);
+}
+
+void FGLTFAnalyticsBuilder::RecordAnimSequence(const UAnimSequence* Object)
+{
+	AnimSequencesRecorded.Add(Object);
+}
+
+void FGLTFAnalyticsBuilder::RecordMaterial(const UMaterialInterface* Object)
+{
+	MaterialsRecorded.Add(Object);
+}
+
+void FGLTFAnalyticsBuilder::RecordTexture(const UTexture* Object)
+{
+	TexturesRecorded.Add(Object);
+}
+
+void FGLTFAnalyticsBuilder::RecordCamera(const UCameraComponent* Object)
+{
+	CamerasRecorded.Add(Object);
+}
+
+void FGLTFAnalyticsBuilder::RecordLight(const ULightComponent* Object)
+{
+	LightsRecorded.Add(Object);
+}
\ No newline at end of file
diff --git a/Engine/Plugins/Enterprise/GLTFExporter/Source/GLTFExporter/Private/Builders/GLTFConvertBuilder.cpp b/Engine/Plugins/Enterprise/GLTFExporter/Source/GLTFExporter/Private/Builders/GLTFConvertBuilder.cpp
--- a/Engine/Plugins/Enterprise/GLTFExporter/Source/GLTFExporter/Private/Builders/GLTFConvertBuilder.cpp
+++ b/Engine/Plugins/Enterprise/GLTFExporter/Source/GLTFExporter/Private/Builders/GLTFConvertBuilder.cpp
@@ -8,7 +8,7 @@
 #include "Engine/TextureRenderTarget2D.h"
 
 FGLTFConvertBuilder::FGLTFConvertBuilder(const FString& FileName, const UGLTFExportOptions* ExportOptions, const TSet<AActor*>& SelectedActors)
-	: FGLTFBufferBuilder(FileName, ExportOptions)
+	: FGLTFAnalyticsBuilder(FileName, ExportOptions)
 	, SelectedActors(SelectedActors)
 {
 }
@@ -111,6 +111,8 @@ FGLTFJsonMesh* FGLTFConvertBuilder::AddUniqueMesh(const UStaticMesh* StaticMesh,
 		return nullptr;
 	}
 
+	RecordStaticMesh(StaticMesh);
+
 	return StaticMeshConverter->GetOrAdd(StaticMesh, nullptr, Materials, LODIndex);
 }
 
@@ -121,6 +123,8 @@ FGLTFJsonMesh* FGLTFConvertBuilder::AddUniqueMesh(const USkeletalMesh* SkeletalM
 		return nullptr;
 	}
 
+	RecordSkeletalMesh(SkeletalMesh);
+
 	return SkeletalMeshConverter->GetOrAdd(SkeletalMesh, nullptr, Materials, LODIndex);
 }
 
@@ -152,6 +156,8 @@ FGLTFJsonMesh* FGLTFConvertBuilder::AddUniqueMesh(const UStaticMeshComponent* St
 		return nullptr;
 	}
 
+	RecordStaticMesh(StaticMesh);
+
 	return StaticMeshConverter->GetOrAdd(StaticMesh, StaticMeshComponent, Materials, LODIndex);
 }
 
@@ -168,16 +174,22 @@ FGLTFJsonMesh* FGLTFConvertBuilder::AddUniqueMesh(const USkeletalMeshComponent*
 		return nullptr;
 	}
 
+	RecordSkeletalMesh(SkeletalMesh);
+
 	return SkeletalMeshConverter->GetOrAdd(SkeletalMesh, SkeletalMeshComponent, Materials, LODIndex);
 }
 
 const FGLTFMeshData* FGLTFConvertBuilder::AddUniqueMeshData(const UStaticMesh* StaticMesh, const UStaticMeshComponent* StaticMeshComponent, int32 LODIndex)
 {
+	RecordStaticMesh(StaticMesh);
+
 	return StaticMeshDataConverter->GetOrAdd(StaticMesh, StaticMeshComponent, LODIndex);
 }
 
 const FGLTFMeshData* FGLTFConvertBuilder::AddUniqueMeshData(const USkeletalMesh* SkeletalMesh, const USkeletalMeshComponent* SkeletalMeshComponent, int32 LODIndex)
 {
+	RecordSkeletalMesh(SkeletalMesh);
+
 	return SkeletalMeshDataConverter->GetOrAdd(SkeletalMesh, SkeletalMeshComponent, LODIndex);
 }
 
@@ -237,6 +249,8 @@ FGLTFJsonMaterial* FGLTFConvertBuilder::AddUniqueMaterial(const UMaterialInterfa
 		return nullptr;
 	}
 
+	RecordMaterial(Material);
+
 	return MaterialConverter->GetOrAdd(Material, MeshData, SectionIndices);
 }
 
@@ -297,6 +311,8 @@ FGLTFJsonTexture* FGLTFConvertBuilder::AddUniqueTexture(const UTexture2D* Textur
 		return nullptr;
 	}
 
+	RecordTexture(Texture);
+
 	return Texture2DConverter->GetOrAdd(Texture, bToSRGB);
 }
 
@@ -307,6 +323,8 @@ FGLTFJsonTexture* FGLTFConvertBuilder::AddUniqueTexture(const UTextureRenderTarg
 		return nullptr;
 	}
 
+	RecordTexture(Texture);
+
 	return TextureRenderTarget2DConverter->GetOrAdd(Texture, bToSRGB);
 }
 
@@ -349,6 +367,8 @@ FGLTFJsonAnimation* FGLTFConvertBuilder::AddUniqueAnimation(FGLTFJsonNode* RootN
 		return nullptr;
 	}
 
+	RecordAnimSequence(AnimSequence);
+
 	return AnimationConverter->GetOrAdd(RootNode, SkeletalMesh, AnimSequence);
 }
 
@@ -369,6 +389,8 @@ FGLTFJsonAnimation* FGLTFConvertBuilder::AddUniqueAnimation(const ULevel* Level,
 		return nullptr;
 	}
 
+	RecordLevelSequence(LevelSequence);
+
 	return LevelSequenceConverter->GetOrAdd(Level, LevelSequence);
 }
 
@@ -389,6 +411,11 @@ FGLTFJsonNode* FGLTFConvertBuilder::AddUniqueNode(const AActor* Actor)
 		return nullptr;
 	}
 
+	if (IsSelectedActor(Actor))
+	{
+		RecordActor(Actor);
+	}
+
 	return ActorConverter->GetOrAdd(Actor);
 }
 
@@ -459,6 +486,8 @@ FGLTFJsonCamera* FGLTFConvertBuilder::AddUniqueCamera(const UCameraComponent* Ca
 		return nullptr;
 	}
 
+	RecordCamera(CameraComponent);
+
 	return CameraConverter->GetOrAdd(CameraComponent);
 }
 
@@ -469,6 +498,8 @@ FGLTFJsonLight* FGLTFConvertBuilder::AddUniqueLight(const ULightComponent* Light
 		return nullptr;
 	}
 
+	RecordLight(LightComponent);
+
 	return LightConverter->GetOrAdd(LightComponent);
 }
 
diff --git a/Engine/Plugins/Enterprise/GLTFExporter/Source/GLTFExporter/Private/Exporters/GLTFExporterAnalytics.cpp b/Engine/Plugins/Enterprise/GLTFExporter/Source/GLTFExporter/Private/Exporters/GLTFExporterAnalytics.cpp
--- a/Engine/Plugins/Enterprise/GLTFExporter/Source/GLTFExporter/Private/Exporters/GLTFExporterAnalytics.cpp
+++ b/Engine/Plugins/Enterprise/GLTFExporter/Source/GLTFExporter/Private/Exporters/GLTFExporterAnalytics.cpp
@@ -6,20 +6,21 @@
 #include "EngineAnalytics.h"
 #include "Engine/World.h"
 
-FGLTFExporterAnalytics::FGLTFExporterAnalytics(const UObject* Object, const FGLTFConvertBuilder& Builder, bool bInitiatedByTask, bool bExportAutomated, bool bExportSuccessful, uint64 StartTime)
-	: FGLTFExporterAnalytics(Object, Builder, bInitiatedByTask, bExportAutomated, bExportSuccessful, StartTime, FPlatformTime::Cycles64())
+FGLTFExporterAnalytics::FGLTFExporterAnalytics(const UObject* Object, const FGLTFConvertBuilder& InBuilder, bool bInitiatedByTask, bool bExportAutomated, bool bExportSuccessful, uint64 StartTime)
+	: FGLTFExporterAnalytics(Object, InBuilder, bInitiatedByTask, bExportAutomated, bExportSuccessful, StartTime, FPlatformTime::Cycles64())
 {
 }
 
-FGLTFExporterAnalytics::FGLTFExporterAnalytics(const UObject* Object, const FGLTFConvertBuilder& Builder, bool bInitiatedByTask, bool bExportAutomated, bool bExportSuccessful, uint64 StartTime, uint64 EndTime)
+FGLTFExporterAnalytics::FGLTFExporterAnalytics(const UObject* Object, const FGLTFConvertBuilder& InBuilder, bool bInitiatedByTask, bool bExportAutomated, bool bExportSuccessful, uint64 StartTime, uint64 EndTime)
 	: AssetType(Object != nullptr ? Object->GetClass() : nullptr)
-	, ExportOptions(Builder.ExportOptions)
-	, bExportAsGLB(Builder.bIsGLB)
-	, bSelectedOnly(!Builder.SelectedActors.IsEmpty())
+	, ExportOptions(InBuilder.ExportOptions)
+	, bExportAsGLB(InBuilder.bIsGLB)
+	, bSelectedOnly(!InBuilder.SelectedActors.IsEmpty())
 	, bInitiatedByTask(bInitiatedByTask)
 	, bExportAutomated(bExportAutomated)
 	, bExportSuccessful(bExportSuccessful)
 	, ExportDuration(FPlatformTime::ToSeconds64(EndTime - StartTime))
+	, Builder(InBuilder)
 {
 }
 
@@ -43,6 +44,8 @@ void FGLTFExporterAnalytics::Send() const
 	EventAttributes.Emplace(TEXT("Platform"), FPlatformProperties::IniPlatformName());
 	EventAttributes.Emplace(TEXT("EngineMode"), FPlatformMisc::GetEngineMode());
 
+	EventAttributes.Append(Builder.GenerateAnalytics());
+
 	const FString EventName = bExportSuccessful ? TEXT("GLTFExporter.Export") : TEXT("GLTFExporter.ExportFailure");
 	FEngineAnalytics::GetProvider().RecordEvent(EventName, EventAttributes);
 }
diff --git a/Engine/Plugins/Enterprise/GLTFExporter/Source/GLTFExporter/Private/Exporters/GLTFExporterAnalytics.h b/Engine/Plugins/Enterprise/GLTFExporter/Source/GLTFExporter/Private/Exporters/GLTFExporterAnalytics.h
--- a/Engine/Plugins/Enterprise/GLTFExporter/Source/GLTFExporter/Private/Exporters/GLTFExporterAnalytics.h
+++ b/Engine/Plugins/Enterprise/GLTFExporter/Source/GLTFExporter/Private/Exporters/GLTFExporterAnalytics.h
@@ -19,12 +19,13 @@ struct FGLTFExporterAnalytics
 	bool bExportSuccessful;
 	double ExportDuration;
 
-	FGLTFExporterAnalytics(const UObject* Object, const FGLTFConvertBuilder& Builder, bool bInitiatedByTask, bool bExportAutomated, bool bExportSuccessful, uint64 StartTime);
-	FGLTFExporterAnalytics(const UObject* Object, const FGLTFConvertBuilder& Builder, bool bInitiatedByTask, bool bExportAutomated, bool bExportSuccessful, uint64 StartTime, uint64 EndTime);
+	FGLTFExporterAnalytics(const UObject* Object, const FGLTFConvertBuilder& InBuilder, bool bInitiatedByTask, bool bExportAutomated, bool bExportSuccessful, uint64 StartTime);
+	FGLTFExporterAnalytics(const UObject* Object, const FGLTFConvertBuilder& InBuilder, bool bInitiatedByTask, bool bExportAutomated, bool bExportSuccessful, uint64 StartTime, uint64 EndTime);
 
 	void Send() const;
 
 private:
+	const FGLTFConvertBuilder& Builder;
 
 	static void GetAttributesFromStruct(const UStruct* Struct, const void* ContainerPtr, const FString& AttributeName, TArray<FAnalyticsEventAttribute>& OutAttributes);
 	static void GetAttributesFromProperty(const FProperty* Property, const void* ValuePtr, const FString& AttributeName, TArray<FAnalyticsEventAttribute>& OutAttributes);
diff --git a/Engine/Plugins/Enterprise/GLTFExporter/Source/GLTFExporter/Public/Builders/GLTFAnalyticsBuilder.h b/Engine/Plugins/Enterprise/GLTFExporter/Source/GLTFExporter/Public/Builders/GLTFAnalyticsBuilder.h
new file mode 100644
--- /dev/null
+++ b/Engine/Plugins/Enterprise/GLTFExporter/Source/GLTFExporter/Public/Builders/GLTFAnalyticsBuilder.h
@@ -0,0 +1,53 @@
+// Copyright Epic Games, Inc. All Rights Reserved.
+
+#pragma once
+
+#include "Builders/GLTFBufferBuilder.h"
+#include "CoreMinimal.h"
+
+struct FAnalyticsEventAttribute;
+
+class AActor;
+class UStaticMesh;
+class USkeletalMesh;
+class ULevelSequence;
+class UAnimSequence;
+class UMaterialInterface;
+class UTexture;
+class UCameraComponent;
+class ULightComponent;
+
+class GLTFEXPORTER_API FGLTFAnalyticsBuilder : public FGLTFBufferBuilder
+{
+public:
+
+	FGLTFAnalyticsBuilder(const FString& FileName, const UGLTFExportOptions* ExportOptions);
+
+	TArray<FAnalyticsEventAttribute> GenerateAnalytics() const;
+
+protected:
+
+	void RecordActor(const AActor* Object);
+	void RecordStaticMesh(const UStaticMesh* Object);
+	void RecordSkeletalMesh(const USkeletalMesh* Object);
+	void RecordLevelSequence(const ULevelSequence* Object);
+	void RecordAnimSequence(const UAnimSequence* Object);
+	void RecordMaterial(const UMaterialInterface* Object);
+	void RecordTexture(const UTexture* Object);
+	void RecordCamera(const UCameraComponent* Object);
+	void RecordLight(const ULightComponent* Object);
+
+private:
+
+	TSet<const AActor*>					ActorsRecorded;
+	TSet<const USceneComponent*>		ComponentsRecorded;
+	TSet<const UStaticMesh*>			StaticMeshesRecorded;
+	TSet<const USkeletalMesh*>			SkeletalMeshesRecorded;
+	TSet<const ULevelSequence*>			LevelSequencesRecorded;
+	TSet<const UAnimSequence*>			AnimSequencesRecorded;
+	TSet<const UMaterialInterface*>		MaterialsRecorded;
+	TSet<const UTexture*>				TexturesRecorded;
+	
+	TSet<const UCameraComponent*>		CamerasRecorded;
+	TSet<const ULightComponent*>		LightsRecorded;
+};
diff --git a/Engine/Plugins/Enterprise/GLTFExporter/Source/GLTFExporter/Public/Builders/GLTFConvertBuilder.h b/Engine/Plugins/Enterprise/GLTFExporter/Source/GLTFExporter/Public/Builders/GLTFConvertBuilder.h
--- a/Engine/Plugins/Enterprise/GLTFExporter/Source/GLTFExporter/Public/Builders/GLTFConvertBuilder.h
+++ b/Engine/Plugins/Enterprise/GLTFExporter/Source/GLTFExporter/Public/Builders/GLTFConvertBuilder.h
@@ -2,7 +2,7 @@
 
 #pragma once
 
-#include "Builders/GLTFBufferBuilder.h"
+#include "Builders/GLTFAnalyticsBuilder.h"
 #include "Converters/GLTFAccessorConverters.h"
 #include "Converters/GLTFMeshConverters.h"
 #include "Converters/GLTFMeshDataConverters.h"
@@ -21,7 +21,7 @@
 class UMeshComponent;
 class UPropertyValue;
 
-class GLTFEXPORTER_API FGLTFConvertBuilder : public FGLTFBufferBuilder
+class GLTFEXPORTER_API FGLTFConvertBuilder : public FGLTFAnalyticsBuilder
 {
 public:
 
