Fix compiler warning regarding FReferenceCollector::AddReferencedObjects

diff --git a/Engine/Plugins/Enterprise/GLTFExporter/Source/GLTFExporter/Public/Converters/GLTFConverter.h b/Engine/Plugins/Enterprise/GLTFExporter/Source/GLTFExporter/Public/Converters/GLTFConverter.h
--- a/Engine/Plugins/Enterprise/GLTFExporter/Source/GLTFExporter/Public/Converters/GLTFConverter.h
+++ b/Engine/Plugins/Enterprise/GLTFExporter/Source/GLTFExporter/Public/Converters/GLTFConverter.h
@@ -43,6 +43,11 @@ public:
 		return NewOutput;
 	}
 
+	const TMap<InputTupleType, OutputType>& GetOutputs() const
+	{
+		return CachedOutputs;
+	}
+
 protected:
 
 	virtual void Sanitize(InputTypes&... Inputs) { }
diff --git a/Engine/Plugins/Enterprise/GLTFExporter/Source/GLTFExporter/Public/Converters/GLTFObjectArrayScopeGuard.h b/Engine/Plugins/Enterprise/GLTFExporter/Source/GLTFExporter/Public/Converters/GLTFObjectArrayScopeGuard.h
--- a/Engine/Plugins/Enterprise/GLTFExporter/Source/GLTFExporter/Public/Converters/GLTFObjectArrayScopeGuard.h
+++ b/Engine/Plugins/Enterprise/GLTFExporter/Source/GLTFExporter/Public/Converters/GLTFObjectArrayScopeGuard.h
@@ -6,7 +6,7 @@
 #include "UObject/GCObject.h"
 
 template <typename ObjectType = UObject, typename = typename TEnableIf<TIsDerivedFrom<ObjectType, UObject>::Value>::Type>
-class FGLTFObjectArrayScopeGuard : public FGCObject, public TArray<ObjectType*>
+class FGLTFObjectArrayScopeGuard : public FGCObject, public TArray<TObjectPtr<ObjectType>>
 {
 public:
 
@@ -14,7 +14,7 @@ public:
 	{
 	}
 
-	using TArray<ObjectType*>::TArray;
+	using TArray<TObjectPtr<ObjectType>>::TArray;
 
 	/** Non-copyable */
 	FGLTFObjectArrayScopeGuard(const FGLTFObjectArrayScopeGuard&) = delete;
