[Export - glTF] Material Specular property to glTF (KHR_materials_specular).

diff --git a/Engine/Plugins/Enterprise/GLTFExporter/Source/GLTFExporter/Private/Converters/GLTFMaterialProxyFactory.cpp b/Engine/Plugins/Enterprise/GLTFExporter/Source/GLTFExporter/Private/Converters/GLTFMaterialProxyFactory.cpp
--- a/Engine/Plugins/Enterprise/GLTFExporter/Source/GLTFExporter/Private/Converters/GLTFMaterialProxyFactory.cpp
+++ b/Engine/Plugins/Enterprise/GLTFExporter/Source/GLTFExporter/Private/Converters/GLTFMaterialProxyFactory.cpp
@@ -130,6 +130,12 @@ void FGLTFMaterialProxyFactory::SetProxyParameters(UMaterialInstanceConstant* Pr
 			SetProxyParameter(ProxyMaterial, FGLTFProxyMaterialInfo::ClearCoatNormalScale, JsonMaterial.ClearCoat.ClearCoatNormalTexture.Scale);
 		}
 	}
+
+	if (JsonMaterial.ShadingModel != EGLTFJsonShadingModel::Unlit) //&& != SpecularGlossiness
+	{
+		SetProxyParameter(ProxyMaterial, FGLTFProxyMaterialInfo::SpecularFactor, JsonMaterial.Specular.Factor);
+		SetProxyParameter(ProxyMaterial, FGLTFProxyMaterialInfo::SpecularTexture, JsonMaterial.Specular.Texture);
+	}
 }
 
 void FGLTFMaterialProxyFactory::SetProxyParameter(UMaterialInstanceConstant* ProxyMaterial, const TGLTFProxyMaterialParameterInfo<float>& ParameterInfo, float Scalar)
diff --git a/Engine/Plugins/Enterprise/GLTFExporter/Source/GLTFExporter/Private/Json/GLTFJsonMaterial.cpp b/Engine/Plugins/Enterprise/GLTFExporter/Source/GLTFExporter/Private/Json/GLTFJsonMaterial.cpp
--- a/Engine/Plugins/Enterprise/GLTFExporter/Source/GLTFExporter/Private/Json/GLTFJsonMaterial.cpp
+++ b/Engine/Plugins/Enterprise/GLTFExporter/Source/GLTFExporter/Private/Json/GLTFJsonMaterial.cpp
@@ -120,6 +120,19 @@ void FGLTFJsonClearCoatExtension::WriteObject(IGLTFJsonWriter& Writer) const
 	}
 }
 
+void FGLTFJsonSpecularExtension::WriteObject(IGLTFJsonWriter& Writer) const
+{
+	if (!FMath::IsNearlyEqual(Factor, 1, Writer.DefaultTolerance))
+	{
+		Writer.Write(TEXT("specularFactor"), Factor);
+	}
+
+	if (Texture.Index != nullptr)
+	{
+		Writer.Write(TEXT("specularTexture"), Texture);
+	}
+}
+
 void FGLTFJsonMaterial::WriteObject(IGLTFJsonWriter& Writer) const
 {
 	if (!Name.IsEmpty())
@@ -168,10 +181,15 @@ void FGLTFJsonMaterial::WriteObject(IGLTFJsonWriter& Writer) const
 	}
 
 	const bool HasEmissiveStrength = !FMath::IsNearlyEqual(EmissiveStrength, 1.0f, Writer.DefaultTolerance);
-	if (ShadingModel == EGLTFJsonShadingModel::Unlit || ShadingModel == EGLTFJsonShadingModel::ClearCoat || HasEmissiveStrength)
+	if (ShadingModel == EGLTFJsonShadingModel::Unlit || ShadingModel == EGLTFJsonShadingModel::ClearCoat || HasEmissiveStrength || Specular.HasValue())
 	{
 		Writer.StartExtensions();
 
+		if (ShadingModel != EGLTFJsonShadingModel::Unlit /*&& !KHR_materials_pbrSpecularGlossiness*/ && Specular.HasValue())
+		{
+			Writer.Write(EGLTFJsonExtension::KHR_MaterialsSpecular, Specular);
+		}
+
 		if (ShadingModel == EGLTFJsonShadingModel::Unlit)
 		{
 			// Write empty object
diff --git a/Engine/Plugins/Enterprise/GLTFExporter/Source/GLTFExporter/Private/Materials/GLTFProxyMaterialInfo.cpp b/Engine/Plugins/Enterprise/GLTFExporter/Source/GLTFExporter/Private/Materials/GLTFProxyMaterialInfo.cpp
--- a/Engine/Plugins/Enterprise/GLTFExporter/Source/GLTFExporter/Private/Materials/GLTFProxyMaterialInfo.cpp
+++ b/Engine/Plugins/Enterprise/GLTFExporter/Source/GLTFExporter/Private/Materials/GLTFProxyMaterialInfo.cpp
@@ -27,3 +27,6 @@ const TGLTFProxyMaterialParameterInfo<float> FGLTFProxyMaterialInfo::ClearCoatRo
 
 const FGLTFProxyMaterialTextureParameterInfo FGLTFProxyMaterialInfo::ClearCoatNormal = { TEXT("Clear Coat Normal") };
 const TGLTFProxyMaterialParameterInfo<float> FGLTFProxyMaterialInfo::ClearCoatNormalScale = { TEXT("Clear Coat Normal Scale") };
+
+const TGLTFProxyMaterialParameterInfo<float> FGLTFProxyMaterialInfo::SpecularFactor = { TEXT("Specular Factor") };
+const FGLTFProxyMaterialTextureParameterInfo FGLTFProxyMaterialInfo::SpecularTexture = { TEXT("Specular") };
diff --git a/Engine/Plugins/Enterprise/GLTFExporter/Source/GLTFExporter/Private/Materials/GLTFProxyMaterialInfo.h b/Engine/Plugins/Enterprise/GLTFExporter/Source/GLTFExporter/Private/Materials/GLTFProxyMaterialInfo.h
--- a/Engine/Plugins/Enterprise/GLTFExporter/Source/GLTFExporter/Private/Materials/GLTFProxyMaterialInfo.h
+++ b/Engine/Plugins/Enterprise/GLTFExporter/Source/GLTFExporter/Private/Materials/GLTFProxyMaterialInfo.h
@@ -33,4 +33,7 @@ public:
 
 	static const FGLTFProxyMaterialTextureParameterInfo ClearCoatNormal;
 	static const TGLTFProxyMaterialParameterInfo<float> ClearCoatNormalScale;
+	
+	static const TGLTFProxyMaterialParameterInfo<float> SpecularFactor;
+	static const FGLTFProxyMaterialTextureParameterInfo SpecularTexture; //Only using Alpha Channel
 };
diff --git a/Engine/Plugins/Enterprise/GLTFExporter/Source/GLTFExporter/Private/Tasks/GLTFDelayedMaterialTasks.cpp b/Engine/Plugins/Enterprise/GLTFExporter/Source/GLTFExporter/Private/Tasks/GLTFDelayedMaterialTasks.cpp
--- a/Engine/Plugins/Enterprise/GLTFExporter/Source/GLTFExporter/Private/Tasks/GLTFDelayedMaterialTasks.cpp
+++ b/Engine/Plugins/Enterprise/GLTFExporter/Source/GLTFExporter/Private/Tasks/GLTFDelayedMaterialTasks.cpp
@@ -46,6 +46,7 @@ namespace
 	const TArray<FLinearColor> OcclusionInputMasks = { OcclusionMask };
 	const TArray<FLinearColor> ClearCoatInputMasks = { ClearCoatMask };
 	const TArray<FLinearColor> ClearCoatRoughnessInputMasks = { ClearCoatRoughnessMask };
+	const TArray<FLinearColor> SpecularMasks = { RedMask, GreenMask, BlueMask, AlphaMask };
 }
 
 void FGLTFDelayedMaterialTask::Process()
@@ -140,6 +141,15 @@ void FGLTFDelayedMaterialTask::Process()
 				}
 			}
 
+			const FMaterialPropertyEx SpecularProperty = MP_Specular;
+			if (IsPropertyNonDefault(SpecularProperty))
+			{
+				if (!TryGetSpecular(*JsonMaterial, SpecularProperty))
+				{
+					Builder.LogWarning(FString::Printf(TEXT("Failed to export %s for material %s"), *SpecularProperty.ToString(), *Material->GetName()));
+				}
+			}
+
 			if (JsonMaterial->ShadingModel == EGLTFJsonShadingModel::ClearCoat)
 			{
 				const FMaterialPropertyEx ClearCoatProperty = MP_CustomData0;
@@ -257,6 +267,12 @@ void FGLTFDelayedMaterialTask::GetProxyParameters(FGLTFJsonMaterial& OutMaterial
 			GetProxyParameter(FGLTFProxyMaterialInfo::ClearCoatNormalScale, OutMaterial.ClearCoat.ClearCoatNormalTexture.Scale);
 		}
 	}
+
+	if (OutMaterial.ShadingModel != EGLTFJsonShadingModel::Unlit)
+	{
+		GetProxyParameter(FGLTFProxyMaterialInfo::SpecularFactor, OutMaterial.Specular.Factor);
+		GetProxyParameter(FGLTFProxyMaterialInfo::SpecularTexture, OutMaterial.Specular.Texture);
+	}
 }
 
 void FGLTFDelayedMaterialTask::GetProxyParameter(const TGLTFProxyMaterialParameterInfo<float>& ParameterInfo, float& OutValue) const
@@ -694,7 +710,7 @@ bool FGLTFDelayedMaterialTask::TryGetMetallicAndRoughness(FGLTFJsonPBRMetallicRo
 		Builder,
 		CombinedPixels,
 		TextureSize,
-		true, // NOTE: we can ignore alpha in everything but TryGetBaseColorAndOpacity
+		true, // NOTE: we can ignore alpha in everything but TryGetBaseColorAndOpacity and TryGetSpecular
 		false,
 		GetBakedTextureName(TEXT("MetallicRoughness")),
 		TextureAddress,
@@ -912,6 +928,138 @@ bool FGLTFDelayedMaterialTask::TryGetEmissive(FGLTFJsonMaterial& OutMaterial, co
 	return true;
 }
 
+bool FGLTFDelayedMaterialTask::TryGetSpecular(FGLTFJsonMaterial& OutMaterial, const FMaterialPropertyEx& SpecularProperty)
+{
+	float Factor;
+	if (TryGetConstantScalar(Factor, MP_Specular))
+	{
+		OutMaterial.Specular.Factor = Factor;
+	}
+	else if (TryGetSourceTexture(OutMaterial.Specular.Texture, SpecularProperty, { AlphaMask }))
+	{
+		OutMaterial.Specular.Factor = 1.0f;
+	}
+	else
+	{
+		if (Builder.ExportOptions->BakeMaterialInputs == EGLTFMaterialBakeMode::Disabled)
+		{
+			Builder.LogWarning(FString::Printf(
+				TEXT("%s for material %s needs to bake, but material baking is disabled by export options"),
+				*SpecularProperty.ToString(),
+				*Material->GetName()));
+			return false;
+		}
+
+		FGLTFPropertyBakeOutput PropertyBakeOutput = BakeMaterialProperty(SpecularProperty, OutMaterial.Specular.Texture.TexCoord, OutMaterial.Specular.Texture.Transform);
+
+		if (PropertyBakeOutput.bIsConstant)
+		{
+			OutMaterial.Specular.Factor = PropertyBakeOutput.ConstantValue.R;
+		}
+		else
+		{
+			if (Builder.ExportOptions->TextureImageFormat == EGLTFTextureImageFormat::None)
+			{
+				OutMaterial.Specular.Texture.Index = nullptr;
+				return true;
+			}
+
+			//Move Value to Alpha per documentation:
+			//
+			{
+				const FExpressionInput* MaterialInput = FGLTFMaterialUtilities::GetInputForProperty(Material, SpecularProperty);
+				if (MaterialInput == nullptr)
+				{
+					// TODO: report error
+					return false;
+				}
+				const FLinearColor Mask = FGLTFMaterialUtilities::GetMask(*MaterialInput);
+
+				enum ERGBMask
+				{
+					RED = 0,
+					GREEN = 1,
+					BLUE = 2,
+					ALPHA = 3
+				};
+
+				ERGBMask IndexOfMask = ERGBMask(0);
+
+				if (SpecularMasks.Contains(Mask))
+				{
+					IndexOfMask = ERGBMask(SpecularMasks.IndexOfByKey(Mask));
+				}
+
+				auto SwapValues = [](FColor& Color) {};
+
+				switch (IndexOfMask)
+				{
+					case RED:
+						for (FColor& Pixel : *PropertyBakeOutput.Pixels)
+						{
+							Pixel.A = Pixel.R;
+
+							Pixel.R = 255;
+							Pixel.G = 255;
+							Pixel.B = 255;
+						}
+						break;
+
+					case GREEN:
+						for (FColor& Pixel : *PropertyBakeOutput.Pixels)
+						{
+							Pixel.A = Pixel.G;
+
+							Pixel.R = 255;
+							Pixel.G = 255;
+							Pixel.B = 255;
+						}
+						break;
+
+					case BLUE:
+						for (FColor& Pixel : *PropertyBakeOutput.Pixels)
+						{
+							Pixel.A = Pixel.B;
+
+							Pixel.R = 255;
+							Pixel.G = 255;
+							Pixel.B = 255;
+						}
+						break;
+
+					case ALPHA:
+					default:
+						break;
+				}
+			}
+
+			const EGLTFMaterialPropertyGroup PropertyGroup = GetPropertyGroup(SpecularProperty);
+			const TextureAddress TextureAddress = Builder.GetBakeTilingForMaterialProperty(Material, PropertyGroup);
+			const TextureFilter TextureFilter = Builder.GetBakeFilterForMaterialProperty(Material, PropertyGroup);
+
+			FGLTFJsonTexture* Texture = FGLTFMaterialUtilities::AddTexture(
+				Builder,
+				PropertyBakeOutput.Pixels,
+				PropertyBakeOutput.Size,
+				false,
+				FGLTFMaterialUtilities::IsNormalMap(SpecularProperty),
+				GetBakedTextureName(TEXT("Specular")),
+				TextureAddress,
+				TextureFilter);
+
+			OutMaterial.Specular.Texture.Index = Texture;
+
+			if (!PropertyBakeOutput.bIsConstant)
+			{
+				OutMaterial.Specular.Factor = 1.0f;
+			}
+		}
+	}
+
+
+	return true;
+}
+
 bool FGLTFDelayedMaterialTask::IsPropertyNonDefault(const FMaterialPropertyEx& Property) const
 {
 	// Custom outputs are by definition standalone and never part of material attributes
@@ -1188,7 +1336,7 @@ bool FGLTFDelayedMaterialTask::TryGetSourceTexture(FGLTFJsonTextureInfo& OutTexI
 
 	if (TryGetSourceTexture(Texture, TexCoord, Transform, Property, AllowedMasks))
 	{
-		OutTexInfo.Index = Builder.AddUniqueTexture(Texture, FGLTFMaterialUtilities::IsSRGB(Property));
+		OutTexInfo.Index = Builder.AddUniqueTexture(Texture, FGLTFMaterialUtilities::IsSRGB(Property) || (Property == MP_Specular && Texture->SRGB));
 		OutTexInfo.TexCoord = TexCoord;
 		OutTexInfo.Transform = Transform;
 		return true;
diff --git a/Engine/Plugins/Enterprise/GLTFExporter/Source/GLTFExporter/Private/Tasks/GLTFDelayedMaterialTasks.h b/Engine/Plugins/Enterprise/GLTFExporter/Source/GLTFExporter/Private/Tasks/GLTFDelayedMaterialTasks.h
--- a/Engine/Plugins/Enterprise/GLTFExporter/Source/GLTFExporter/Private/Tasks/GLTFDelayedMaterialTasks.h
+++ b/Engine/Plugins/Enterprise/GLTFExporter/Source/GLTFExporter/Private/Tasks/GLTFDelayedMaterialTasks.h
@@ -75,6 +75,7 @@ private:
 	bool TryGetMetallicAndRoughness(FGLTFJsonPBRMetallicRoughness& OutPBRParams, const FMaterialPropertyEx& MetallicProperty, const FMaterialPropertyEx& RoughnessProperty);
 	bool TryGetClearCoatRoughness(FGLTFJsonClearCoatExtension& OutExtParams, const FMaterialPropertyEx& IntensityProperty, const FMaterialPropertyEx& RoughnessProperty);
 	bool TryGetEmissive(FGLTFJsonMaterial& OutMaterial, const FMaterialPropertyEx& EmissiveProperty);
+	bool TryGetSpecular(FGLTFJsonMaterial& OutMaterial, const FMaterialPropertyEx& SpecularProperty);
 
 	bool IsPropertyNonDefault(const FMaterialPropertyEx& Property) const;
 	bool TryGetConstantColor(FGLTFJsonColor3& OutValue, const FMaterialPropertyEx& Property) const;
diff --git a/Engine/Plugins/Enterprise/GLTFExporter/Source/GLTFExporter/Private/Utilities/GLTFJsonUtilities.cpp b/Engine/Plugins/Enterprise/GLTFExporter/Source/GLTFExporter/Private/Utilities/GLTFJsonUtilities.cpp
--- a/Engine/Plugins/Enterprise/GLTFExporter/Source/GLTFExporter/Private/Utilities/GLTFJsonUtilities.cpp
+++ b/Engine/Plugins/Enterprise/GLTFExporter/Source/GLTFExporter/Private/Utilities/GLTFJsonUtilities.cpp
@@ -13,6 +13,7 @@ const TCHAR* FGLTFJsonUtilities::GetValue(EGLTFJsonExtension Enum)
 		case EGLTFJsonExtension::KHR_MaterialsVariants:         return TEXT("KHR_materials_variants");
 		case EGLTFJsonExtension::KHR_MeshQuantization:          return TEXT("KHR_mesh_quantization");
 		case EGLTFJsonExtension::KHR_TextureTransform:          return TEXT("KHR_texture_transform");
+		case EGLTFJsonExtension::KHR_MaterialsSpecular:	        return TEXT("KHR_materials_specular");
 		default:
 			checkNoEntry();
 			return TEXT("");
diff --git a/Engine/Plugins/Enterprise/GLTFExporter/Source/GLTFExporter/Public/Json/GLTFJsonEnums.h b/Engine/Plugins/Enterprise/GLTFExporter/Source/GLTFExporter/Public/Json/GLTFJsonEnums.h
--- a/Engine/Plugins/Enterprise/GLTFExporter/Source/GLTFExporter/Public/Json/GLTFJsonEnums.h
+++ b/Engine/Plugins/Enterprise/GLTFExporter/Source/GLTFExporter/Public/Json/GLTFJsonEnums.h
@@ -10,7 +10,8 @@ enum class EGLTFJsonExtension
 	KHR_MaterialsUnlit,
 	KHR_MaterialsVariants,
 	KHR_MeshQuantization,
-	KHR_TextureTransform
+	KHR_TextureTransform,
+	KHR_MaterialsSpecular
 };
 
 enum class EGLTFJsonShadingModel
diff --git a/Engine/Plugins/Enterprise/GLTFExporter/Source/GLTFExporter/Public/Json/GLTFJsonMaterial.h b/Engine/Plugins/Enterprise/GLTFExporter/Source/GLTFExporter/Public/Json/GLTFJsonMaterial.h
--- a/Engine/Plugins/Enterprise/GLTFExporter/Source/GLTFExporter/Public/Json/GLTFJsonMaterial.h
+++ b/Engine/Plugins/Enterprise/GLTFExporter/Source/GLTFExporter/Public/Json/GLTFJsonMaterial.h
@@ -83,6 +83,25 @@ struct GLTFEXPORTER_API FGLTFJsonClearCoatExtension : IGLTFJsonObject
 	virtual void WriteObject(IGLTFJsonWriter& Writer) const override;
 };
 
+/*
+* Only taking advantage of SpecularFactor and SpecularMaps
+*/
+struct GLTFEXPORTER_API FGLTFJsonSpecularExtension : IGLTFJsonObject
+{
+	float Factor; //SpecularFactor
+	FGLTFJsonTextureInfo Texture; //SpecularMap
+
+	FGLTFJsonSpecularExtension()
+		: Factor(0.5)
+	{
+	}
+
+	//UE's default Specular value is 0.5, while glTF's is 1.
+	bool HasValue() const { return ((Factor != 0.5f && Factor != 1.0f) || Texture.Index != nullptr); }
+
+	virtual void WriteObject(IGLTFJsonWriter& Writer) const override;
+};
+
 struct GLTFEXPORTER_API FGLTFJsonMaterial : IGLTFJsonIndexedObject
 {
 	FString Name;
@@ -105,6 +124,8 @@ struct GLTFEXPORTER_API FGLTFJsonMaterial : IGLTFJsonIndexedObject
 
 	FGLTFJsonClearCoatExtension ClearCoat;
 
+	FGLTFJsonSpecularExtension Specular;
+
 	virtual void WriteObject(IGLTFJsonWriter& Writer) const override;
 
 protected:
