[Export - glTF] Material Refraction property to glTF (KHR_materials_ior)

diff --git a/Engine/Plugins/Enterprise/GLTFExporter/Source/GLTFExporter/Private/Json/GLTFJsonMaterial.cpp b/Engine/Plugins/Enterprise/GLTFExporter/Source/GLTFExporter/Private/Json/GLTFJsonMaterial.cpp
--- a/Engine/Plugins/Enterprise/GLTFExporter/Source/GLTFExporter/Private/Json/GLTFJsonMaterial.cpp
+++ b/Engine/Plugins/Enterprise/GLTFExporter/Source/GLTFExporter/Private/Json/GLTFJsonMaterial.cpp
@@ -133,6 +133,14 @@ void FGLTFJsonSpecularExtension::WriteObject(IGLTFJsonWriter& Writer) const
 	}
 }
 
+void FGLTFJsonIORExtension::WriteObject(IGLTFJsonWriter& Writer) const
+{
+	if (!FMath::IsNearlyEqual(Value, 1.5, Writer.DefaultTolerance))
+	{
+		Writer.Write(TEXT("ior"), Value);
+	}
+}
+
 void FGLTFJsonMaterial::WriteObject(IGLTFJsonWriter& Writer) const
 {
 	if (!Name.IsEmpty())
@@ -181,7 +189,7 @@ void FGLTFJsonMaterial::WriteObject(IGLTFJsonWriter& Writer) const
 	}
 
 	const bool HasEmissiveStrength = !FMath::IsNearlyEqual(EmissiveStrength, 1.0f, Writer.DefaultTolerance);
-	if (ShadingModel == EGLTFJsonShadingModel::Unlit || ShadingModel == EGLTFJsonShadingModel::ClearCoat || HasEmissiveStrength || Specular.HasValue())
+	if (ShadingModel == EGLTFJsonShadingModel::Unlit || ShadingModel == EGLTFJsonShadingModel::ClearCoat || HasEmissiveStrength || Specular.HasValue() || IOR.HasValue())
 	{
 		Writer.StartExtensions();
 
@@ -190,6 +198,11 @@ void FGLTFJsonMaterial::WriteObject(IGLTFJsonWriter& Writer) const
 			Writer.Write(EGLTFJsonExtension::KHR_MaterialsSpecular, Specular);
 		}
 
+		if (ShadingModel != EGLTFJsonShadingModel::Unlit /*&& !KHR_materials_pbrSpecularGlossiness*/ && IOR.HasValue())
+		{
+			Writer.Write(EGLTFJsonExtension::KHR_MaterialsIOR, IOR);
+		}
+
 		if (ShadingModel == EGLTFJsonShadingModel::Unlit)
 		{
 			// Write empty object
diff --git a/Engine/Plugins/Enterprise/GLTFExporter/Source/GLTFExporter/Private/Tasks/GLTFDelayedMaterialTasks.cpp b/Engine/Plugins/Enterprise/GLTFExporter/Source/GLTFExporter/Private/Tasks/GLTFDelayedMaterialTasks.cpp
--- a/Engine/Plugins/Enterprise/GLTFExporter/Source/GLTFExporter/Private/Tasks/GLTFDelayedMaterialTasks.cpp
+++ b/Engine/Plugins/Enterprise/GLTFExporter/Source/GLTFExporter/Private/Tasks/GLTFDelayedMaterialTasks.cpp
@@ -150,6 +150,21 @@ void FGLTFDelayedMaterialTask::Process()
 				}
 			}
 
+			if (JsonMaterial->AlphaMode == EGLTFJsonAlphaMode::Blend)
+			{
+				if (BaseMaterial->RefractionMethod == ERefractionMode::RM_IndexOfRefraction)
+				{
+					const FMaterialPropertyEx RefractionProperty = MP_Refraction;
+					if (IsPropertyNonDefault(RefractionProperty))
+					{
+						if (!TryGetRefraction(*JsonMaterial, RefractionProperty))
+						{
+							Builder.LogWarning(FString::Printf(TEXT("Failed to export %s for material %s"), *RefractionProperty.ToString(), *Material->GetName()));
+						}
+					}
+				}
+			}
+
 			if (JsonMaterial->ShadingModel == EGLTFJsonShadingModel::ClearCoat)
 			{
 				const FMaterialPropertyEx ClearCoatProperty = MP_CustomData0;
@@ -1060,6 +1075,56 @@ bool FGLTFDelayedMaterialTask::TryGetSpecular(FGLTFJsonMaterial& OutMaterial, co
 	return true;
 }
 
+bool FGLTFDelayedMaterialTask::TryGetRefraction(FGLTFJsonMaterial& OutMaterial, const FMaterialPropertyEx& RefractionProperty)
+{
+	if (!TryGetConstantScalar(OutMaterial.IOR.Value, RefractionProperty))
+	{
+		if (Builder.ExportOptions->BakeMaterialInputs == EGLTFMaterialBakeMode::Disabled)
+		{
+			Builder.LogWarning(FString::Printf(
+				TEXT("%s for material %s needs to bake, but material baking is disabled by export options"),
+				*RefractionProperty.ToString(),
+				*Material->GetName()));
+			return false;
+		}
+
+		int32 TempTexCoord;
+		FGLTFJsonTextureTransform TempTransform;
+		const FIntPoint TextureSize(128);
+		//Note: Refraction baking limits range to [1, Infinity]
+		FGLTFPropertyBakeOutput PropertyBakeOutput = BakeMaterialProperty(RefractionProperty, TempTexCoord, TempTransform, TextureSize, false);
+		if (PropertyBakeOutput.bIsConstant)
+		{
+			OutMaterial.IOR.Value = PropertyBakeOutput.ConstantValue.R;
+		}
+		else
+		{
+			//Calculate an avarage constant value from baked texture:
+			Builder.LogWarning(FString::Printf(
+				TEXT("In material %s : The %s property is using non const value. Avarage const value approximation is used as the glTF 2.0 standard only supports const value for said property (via khr_materials_ior). "),
+				*Material->GetName(),
+				*RefractionProperty.ToString()));
+
+			uint64 Red, Green, Blue;
+			Red = Green = Blue = 0;
+			for (const FColor& Pixel : *PropertyBakeOutput.Pixels)
+			{
+				Red += Pixel.R;
+				Green += Pixel.G;
+				Blue += Pixel.B;
+			}
+
+			OutMaterial.IOR.Value = (Red + Green + Blue) / (3 * PropertyBakeOutput.Pixels.Get().Num() * 255.);
+		}
+
+		//Baking shifts the Range of Refraction from, to:
+		// [1,Infinity] -> [1,0]
+		OutMaterial.IOR.Value = 1. / OutMaterial.IOR.Value;
+	}
+
+	return true;
+}
+
 bool FGLTFDelayedMaterialTask::IsPropertyNonDefault(const FMaterialPropertyEx& Property) const
 {
 	// Custom outputs are by definition standalone and never part of material attributes
diff --git a/Engine/Plugins/Enterprise/GLTFExporter/Source/GLTFExporter/Private/Tasks/GLTFDelayedMaterialTasks.h b/Engine/Plugins/Enterprise/GLTFExporter/Source/GLTFExporter/Private/Tasks/GLTFDelayedMaterialTasks.h
--- a/Engine/Plugins/Enterprise/GLTFExporter/Source/GLTFExporter/Private/Tasks/GLTFDelayedMaterialTasks.h
+++ b/Engine/Plugins/Enterprise/GLTFExporter/Source/GLTFExporter/Private/Tasks/GLTFDelayedMaterialTasks.h
@@ -76,6 +76,7 @@ private:
 	bool TryGetClearCoatRoughness(FGLTFJsonClearCoatExtension& OutExtParams, const FMaterialPropertyEx& IntensityProperty, const FMaterialPropertyEx& RoughnessProperty);
 	bool TryGetEmissive(FGLTFJsonMaterial& OutMaterial, const FMaterialPropertyEx& EmissiveProperty);
 	bool TryGetSpecular(FGLTFJsonMaterial& OutMaterial, const FMaterialPropertyEx& SpecularProperty);
+	bool TryGetRefraction(FGLTFJsonMaterial& OutMaterial, const FMaterialPropertyEx& RefractionProperty);
 
 	bool IsPropertyNonDefault(const FMaterialPropertyEx& Property) const;
 	bool TryGetConstantColor(FGLTFJsonColor3& OutValue, const FMaterialPropertyEx& Property) const;
diff --git a/Engine/Plugins/Enterprise/GLTFExporter/Source/GLTFExporter/Private/Utilities/GLTFJsonUtilities.cpp b/Engine/Plugins/Enterprise/GLTFExporter/Source/GLTFExporter/Private/Utilities/GLTFJsonUtilities.cpp
--- a/Engine/Plugins/Enterprise/GLTFExporter/Source/GLTFExporter/Private/Utilities/GLTFJsonUtilities.cpp
+++ b/Engine/Plugins/Enterprise/GLTFExporter/Source/GLTFExporter/Private/Utilities/GLTFJsonUtilities.cpp
@@ -11,6 +11,7 @@ const TCHAR* FGLTFJsonUtilities::GetValue(EGLTFJsonExtension Enum)
 		case EGLTFJsonExtension::KHR_MaterialsEmissiveStrength: return TEXT("KHR_materials_emissive_strength");
 		case EGLTFJsonExtension::KHR_MaterialsUnlit:            return TEXT("KHR_materials_unlit");
 		case EGLTFJsonExtension::KHR_MaterialsVariants:         return TEXT("KHR_materials_variants");
+		case EGLTFJsonExtension::KHR_MaterialsIOR:              return TEXT("KHR_materials_ior");
 		case EGLTFJsonExtension::KHR_MeshQuantization:          return TEXT("KHR_mesh_quantization");
 		case EGLTFJsonExtension::KHR_TextureTransform:          return TEXT("KHR_texture_transform");
 		case EGLTFJsonExtension::KHR_MaterialsSpecular:	        return TEXT("KHR_materials_specular");
diff --git a/Engine/Plugins/Enterprise/GLTFExporter/Source/GLTFExporter/Public/Json/GLTFJsonEnums.h b/Engine/Plugins/Enterprise/GLTFExporter/Source/GLTFExporter/Public/Json/GLTFJsonEnums.h
--- a/Engine/Plugins/Enterprise/GLTFExporter/Source/GLTFExporter/Public/Json/GLTFJsonEnums.h
+++ b/Engine/Plugins/Enterprise/GLTFExporter/Source/GLTFExporter/Public/Json/GLTFJsonEnums.h
@@ -9,6 +9,7 @@ enum class EGLTFJsonExtension
 	KHR_MaterialsEmissiveStrength,
 	KHR_MaterialsUnlit,
 	KHR_MaterialsVariants,
+	KHR_MaterialsIOR,
 	KHR_MeshQuantization,
 	KHR_TextureTransform,
 	KHR_MaterialsSpecular
diff --git a/Engine/Plugins/Enterprise/GLTFExporter/Source/GLTFExporter/Public/Json/GLTFJsonMaterial.h b/Engine/Plugins/Enterprise/GLTFExporter/Source/GLTFExporter/Public/Json/GLTFJsonMaterial.h
--- a/Engine/Plugins/Enterprise/GLTFExporter/Source/GLTFExporter/Public/Json/GLTFJsonMaterial.h
+++ b/Engine/Plugins/Enterprise/GLTFExporter/Source/GLTFExporter/Public/Json/GLTFJsonMaterial.h
@@ -97,7 +97,21 @@ struct GLTFEXPORTER_API FGLTFJsonSpecularExtension : IGLTFJsonObject
 	}
 
 	//UE's default Specular value is 0.5, while glTF's is 1.
-	bool HasValue() const { return ((Factor != 0.5f && Factor != 1.0f) || Texture.Index != nullptr); }
+	bool HasValue() const { return ((!FMath::IsNearlyEqual(Factor, 0.5f) && !FMath::IsNearlyEqual(Factor, 1.0f)) || Texture.Index != nullptr); }
+
+	virtual void WriteObject(IGLTFJsonWriter& Writer) const override;
+};
+
+struct GLTFEXPORTER_API FGLTFJsonIORExtension : IGLTFJsonObject
+{
+	float Value;
+
+	FGLTFJsonIORExtension()
+		: Value(1.5)
+	{
+	}
+
+	bool HasValue() const { return !FMath::IsNearlyEqual(Value, 1.5f); }
 
 	virtual void WriteObject(IGLTFJsonWriter& Writer) const override;
 };
@@ -123,6 +137,7 @@ struct GLTFEXPORTER_API FGLTFJsonMaterial : IGLTFJsonIndexedObject
 	bool DoubleSided;
 
 	FGLTFJsonClearCoatExtension ClearCoat;
+	FGLTFJsonIORExtension IOR;
 
 	FGLTFJsonSpecularExtension Specular;
 
