[GLTFExporter] Add option UGLTFExportOptions:bMakeSkinnedMeshesRoot (default: true) to avoid local or inherited transforms on skinned meshes (and strictly comply with glTF spec)

Because this affects the scene hierarchy it's desirable to make it an option. Only one viewer (Babylon.js) is known to strictly comply with this requirement in the spec. Not even Khronos sample viewer does it.

#rnx

diff --git a/Engine/Plugins/Enterprise/GLTFExporter/Source/GLTFExporter/Private/Converters/GLTFBoneUtilities.cpp b/Engine/Plugins/Enterprise/GLTFExporter/Source/GLTFExporter/Private/Converters/GLTFBoneUtilities.cpp
--- a/Engine/Plugins/Enterprise/GLTFExporter/Source/GLTFExporter/Private/Converters/GLTFBoneUtilities.cpp
+++ b/Engine/Plugins/Enterprise/GLTFExporter/Source/GLTFExporter/Private/Converters/GLTFBoneUtilities.cpp
@@ -17,7 +17,7 @@ FTransform FGLTFBoneUtilities::GetBindTransform(const FReferenceSkeleton& Refere
 
 	while (ParentBoneIndex != INDEX_NONE)
 	{
-		BindTransform = BindTransform * BonePoses[ParentBoneIndex];
+		BindTransform *= BonePoses[ParentBoneIndex];
 		ParentBoneIndex = BoneInfos[ParentBoneIndex].ParentIndex;
 	}
 
diff --git a/Engine/Plugins/Enterprise/GLTFExporter/Source/GLTFExporter/Private/Converters/GLTFSceneConverters.cpp b/Engine/Plugins/Enterprise/GLTFExporter/Source/GLTFExporter/Private/Converters/GLTFSceneConverters.cpp
--- a/Engine/Plugins/Enterprise/GLTFExporter/Source/GLTFExporter/Private/Converters/GLTFSceneConverters.cpp
+++ b/Engine/Plugins/Enterprise/GLTFExporter/Source/GLTFExporter/Private/Converters/GLTFSceneConverters.cpp
@@ -52,14 +52,19 @@ FGLTFJsonScene* FGLTFSceneConverter::Convert(const UWorld* World)
 
 			for (const AActor* Actor : LevelActors)
 			{
-				FGLTFJsonNode* JsonNode = Builder.AddUniqueNode(Actor);
-				if (JsonNode != nullptr && Builder.IsRootActor(Actor))
+				FGLTFJsonNode* Node = Builder.AddUniqueNode(Actor);
+				if (Node != nullptr && Builder.IsRootActor(Actor))
 				{
 					// TODO: to avoid having to add irrelevant actors/components let GLTFComponentConverter decide and add root nodes to Scene->
 					// This change may require node converters to support cyclic calls.
-					Scene->Nodes.Add(JsonNode);
+					Scene->Nodes.Add(Node);
 				}
 			}
+
+			if (Builder.ExportOptions->bMakeSkinnedMeshesRoot)
+			{
+				MakeSkinnedMeshesRoot(Scene);
+			}
 		}
 	}
 	else
@@ -71,3 +76,36 @@ FGLTFJsonScene* FGLTFSceneConverter::Convert(const UWorld* World)
 
 	return Scene;
 }
+
+void FGLTFSceneConverter::MakeSkinnedMeshesRoot(FGLTFJsonScene* Scene)
+{
+	TArray<FGLTFJsonNode*> OriginalRootNodes = Scene->Nodes;
+
+	for (FGLTFJsonNode* RootNode : OriginalRootNodes)
+	{
+		MakeSkinnedMeshesRoot(RootNode, true, Scene);
+	}
+}
+
+void FGLTFSceneConverter::MakeSkinnedMeshesRoot(FGLTFJsonNode* Node, bool bIsRootNode, FGLTFJsonScene* Scene)
+{
+	if (Node->Mesh != nullptr && Node->Skin != nullptr && (!bIsRootNode || *Node != FGLTFJsonTransform::Identity))
+	{
+		// The glTF spec states "the transform of the skinned mesh node MUST be ignored" (only joint transforms are applied).
+		// To comply with this requirement we can move the skinned mesh to a unique root node that doesn't have any transform.
+		// The end result will be the same, since final joint transforms are unaffected and still applied to the skinned mesh.
+
+		FGLTFJsonNode* SkinnedMeshNode = Builder.AddNode();
+		SkinnedMeshNode->Mesh = Node->Mesh;
+		SkinnedMeshNode->Skin = Node->Skin;
+
+		Node->Mesh = nullptr;
+		Node->Skin = nullptr;
+		Scene->Nodes.Add(SkinnedMeshNode);
+	}
+
+	for (FGLTFJsonNode* ChildNode : Node->Children)
+	{
+		MakeSkinnedMeshesRoot(ChildNode, false, Scene);
+	}
+}
diff --git a/Engine/Plugins/Enterprise/GLTFExporter/Source/GLTFExporter/Private/Json/GLTFJsonNode.cpp b/Engine/Plugins/Enterprise/GLTFExporter/Source/GLTFExporter/Private/Json/GLTFJsonNode.cpp
--- a/Engine/Plugins/Enterprise/GLTFExporter/Source/GLTFExporter/Private/Json/GLTFJsonNode.cpp
+++ b/Engine/Plugins/Enterprise/GLTFExporter/Source/GLTFExporter/Private/Json/GLTFJsonNode.cpp
@@ -13,20 +13,7 @@ void FGLTFJsonNode::WriteObject(IGLTFJsonWriter& Writer) const
 		Writer.Write(TEXT("name"), Name);
 	}
 
-	if (!Translation.IsNearlyEqual(FGLTFJsonVector3::Zero, Writer.DefaultTolerance))
-	{
-		Writer.Write(TEXT("translation"), Translation);
-	}
-
-	if (!Rotation.IsNearlyEqual(FGLTFJsonQuaternion::Identity, Writer.DefaultTolerance))
-	{
-		Writer.Write(TEXT("rotation"), Rotation);
-	}
-
-	if (!Scale.IsNearlyEqual(FGLTFJsonVector3::One, Writer.DefaultTolerance))
-	{
-		Writer.Write(TEXT("scale"), Scale);
-	}
+	FGLTFJsonTransform::WriteValue(Writer);
 
 	if (Camera != nullptr)
 	{
diff --git a/Engine/Plugins/Enterprise/GLTFExporter/Source/GLTFExporter/Private/Json/GLTFJsonTransform.cpp b/Engine/Plugins/Enterprise/GLTFExporter/Source/GLTFExporter/Private/Json/GLTFJsonTransform.cpp
new file mode 100644
--- /dev/null
+++ b/Engine/Plugins/Enterprise/GLTFExporter/Source/GLTFExporter/Private/Json/GLTFJsonTransform.cpp
@@ -0,0 +1,23 @@
+// Copyright Epic Games, Inc. All Rights Reserved.
+
+#include "Json/GLTFJsonTransform.h"
+
+const FGLTFJsonTransform FGLTFJsonTransform::Identity;
+
+void FGLTFJsonTransform::WriteValue(IGLTFJsonWriter& Writer) const
+{
+	if (!Translation.IsNearlyEqual(FGLTFJsonVector3::Zero, Writer.DefaultTolerance))
+	{
+		Writer.Write(TEXT("translation"), Translation);
+	}
+
+	if (!Rotation.IsNearlyEqual(FGLTFJsonQuaternion::Identity, Writer.DefaultTolerance))
+	{
+		Writer.Write(TEXT("rotation"), Rotation);
+	}
+
+	if (!Scale.IsNearlyEqual(FGLTFJsonVector3::One, Writer.DefaultTolerance))
+	{
+		Writer.Write(TEXT("scale"), Scale);
+	}
+}
diff --git a/Engine/Plugins/Enterprise/GLTFExporter/Source/GLTFExporter/Private/Options/GLTFExportOptions.cpp b/Engine/Plugins/Enterprise/GLTFExporter/Source/GLTFExporter/Private/Options/GLTFExportOptions.cpp
--- a/Engine/Plugins/Enterprise/GLTFExporter/Source/GLTFExporter/Private/Options/GLTFExportOptions.cpp
+++ b/Engine/Plugins/Enterprise/GLTFExporter/Source/GLTFExporter/Private/Options/GLTFExportOptions.cpp
@@ -25,6 +25,7 @@ void UGLTFExportOptions::ResetToDefault()
 	DefaultLevelOfDetail = 0;
 	bExportVertexColors = false;
 	bExportVertexSkinWeights = true;
+	bMakeSkinnedMeshesRoot = true;
 	bUseMeshQuantization = false;
 	bExportLevelSequences = true;
 	bExportAnimationSequences = true;
diff --git a/Engine/Plugins/Enterprise/GLTFExporter/Source/GLTFExporter/Public/Converters/GLTFSceneConverters.h b/Engine/Plugins/Enterprise/GLTFExporter/Source/GLTFExporter/Public/Converters/GLTFSceneConverters.h
--- a/Engine/Plugins/Enterprise/GLTFExporter/Source/GLTFExporter/Public/Converters/GLTFSceneConverters.h
+++ b/Engine/Plugins/Enterprise/GLTFExporter/Source/GLTFExporter/Public/Converters/GLTFSceneConverters.h
@@ -17,4 +17,7 @@ public:
 protected:
 
 	virtual FGLTFJsonScene* Convert(const UWorld* Level) override;
+
+	void MakeSkinnedMeshesRoot(FGLTFJsonScene* Scene);
+	void MakeSkinnedMeshesRoot(FGLTFJsonNode* Node, bool bIsRootNode, FGLTFJsonScene* Scene);
 };
diff --git a/Engine/Plugins/Enterprise/GLTFExporter/Source/GLTFExporter/Public/Json/GLTFJsonNode.h b/Engine/Plugins/Enterprise/GLTFExporter/Source/GLTFExporter/Public/Json/GLTFJsonNode.h
--- a/Engine/Plugins/Enterprise/GLTFExporter/Source/GLTFExporter/Public/Json/GLTFJsonNode.h
+++ b/Engine/Plugins/Enterprise/GLTFExporter/Source/GLTFExporter/Public/Json/GLTFJsonNode.h
@@ -3,19 +3,16 @@
 #pragma once
 
 #include "Json/GLTFJsonCore.h"
+#include "Json/GLTFJsonTransform.h"
 
-struct GLTFEXPORTER_API FGLTFJsonNode : IGLTFJsonIndexedObject
+struct GLTFEXPORTER_API FGLTFJsonNode : FGLTFJsonTransform, IGLTFJsonIndexedObject
 {
 	FString Name;
 
-	FGLTFJsonVector3    Translation;
-	FGLTFJsonQuaternion Rotation;
-	FGLTFJsonVector3    Scale;
-
-	FGLTFJsonCamera*    Camera;
-	FGLTFJsonSkin*      Skin;
-	FGLTFJsonMesh*      Mesh;
-	FGLTFJsonLight*     Light;
+	FGLTFJsonCamera* Camera;
+	FGLTFJsonSkin*   Skin;
+	FGLTFJsonMesh*   Mesh;
+	FGLTFJsonLight*  Light;
 
 	TArray<FGLTFJsonNode*> Children;
 
@@ -27,9 +24,6 @@ protected:
 
 	FGLTFJsonNode(int32 Index)
 		: IGLTFJsonIndexedObject(Index)
-		, Translation(FGLTFJsonVector3::Zero)
-		, Rotation(FGLTFJsonQuaternion::Identity)
-		, Scale(FGLTFJsonVector3::One)
 		, Camera(nullptr)
 		, Skin(nullptr)
 		, Mesh(nullptr)
diff --git a/Engine/Plugins/Enterprise/GLTFExporter/Source/GLTFExporter/Public/Json/GLTFJsonTransform.h b/Engine/Plugins/Enterprise/GLTFExporter/Source/GLTFExporter/Public/Json/GLTFJsonTransform.h
new file mode 100644
--- /dev/null
+++ b/Engine/Plugins/Enterprise/GLTFExporter/Source/GLTFExporter/Public/Json/GLTFJsonTransform.h
@@ -0,0 +1,44 @@
+// Copyright Epic Games, Inc. All Rights Reserved.
+
+#pragma once
+
+#include "Json/GLTFJsonCore.h"
+
+struct GLTFEXPORTER_API FGLTFJsonTransform
+{
+	FGLTFJsonVector3    Translation;
+	FGLTFJsonQuaternion Rotation;
+	FGLTFJsonVector3    Scale;
+
+	static const FGLTFJsonTransform Identity;
+
+	FGLTFJsonTransform()
+		: Translation(FGLTFJsonVector3::Zero)
+		, Rotation(FGLTFJsonQuaternion::Identity)
+		, Scale(FGLTFJsonVector3::One)
+	{
+	}
+
+	bool operator==(const FGLTFJsonTransform& Other) const
+	{
+		return Translation == Other.Translation
+			&& Rotation == Other.Rotation
+			&& Scale == Other.Scale;
+	}
+
+	bool operator!=(const FGLTFJsonTransform& Other) const
+	{
+		return Translation != Other.Translation
+			|| Rotation != Other.Rotation
+			|| Scale != Other.Scale;
+	}
+
+	void WriteValue(IGLTFJsonWriter& Writer) const;
+
+	bool IsNearlyEqual(const FGLTFJsonTransform& Other, float Tolerance = KINDA_SMALL_NUMBER) const
+	{
+		return Translation.IsNearlyEqual(Other.Translation, Tolerance)
+			&& Rotation.IsNearlyEqual(Other.Rotation, Tolerance)
+			&& Scale.IsNearlyEqual(Other.Scale, Tolerance);
+	}
+};
diff --git a/Engine/Plugins/Enterprise/GLTFExporter/Source/GLTFExporter/Public/Options/GLTFExportOptions.h b/Engine/Plugins/Enterprise/GLTFExporter/Source/GLTFExporter/Public/Options/GLTFExportOptions.h
--- a/Engine/Plugins/Enterprise/GLTFExporter/Source/GLTFExporter/Public/Options/GLTFExportOptions.h
+++ b/Engine/Plugins/Enterprise/GLTFExporter/Source/GLTFExporter/Public/Options/GLTFExportOptions.h
@@ -97,6 +97,10 @@ class GLTFEXPORTER_API UGLTFExportOptions : public UObject
 	UPROPERTY(EditAnywhere, BlueprintReadWrite, Config, Category = Mesh)
 	bool bExportVertexSkinWeights;
 
+	/** If enabled, make skeletal meshes into root nodes to strictly comply with the glTF specification. Final bone transforms remain the same and visual results are unaffected. */
+	UPROPERTY(EditAnywhere, BlueprintReadWrite, Config, Category = Mesh)
+	bool bMakeSkinnedMeshesRoot;
+
 	/** If enabled, use quantization for vertex tangents and normals, reducing size. Requires extension KHR_mesh_quantization, which may result in the mesh not loading in some glTF viewers. */
 	UPROPERTY(EditAnywhere, BlueprintReadWrite, Config, Category = Mesh)
 	bool bUseMeshQuantization;
