[Export - glTF] Cloth materials to glTF (KHR_materials_sheen)

#rnx

diff --git a/Engine/Plugins/Enterprise/GLTFExporter/Source/GLTFExporter/Private/Converters/GLTFMaterialProxyFactory.cpp b/Engine/Plugins/Enterprise/GLTFExporter/Source/GLTFExporter/Private/Converters/GLTFMaterialProxyFactory.cpp
--- a/Engine/Plugins/Enterprise/GLTFExporter/Source/GLTFExporter/Private/Converters/GLTFMaterialProxyFactory.cpp
+++ b/Engine/Plugins/Enterprise/GLTFExporter/Source/GLTFExporter/Private/Converters/GLTFMaterialProxyFactory.cpp
@@ -102,7 +102,9 @@ void FGLTFMaterialProxyFactory::SetProxyParameters(UMaterialInstanceConstant* Pr
 	SetProxyParameter(ProxyMaterial, FGLTFProxyMaterialInfo::BaseColor, JsonMaterial.PBRMetallicRoughness.BaseColorTexture);
 	SetProxyParameter(ProxyMaterial, FGLTFProxyMaterialInfo::BaseColorFactor, JsonMaterial.PBRMetallicRoughness.BaseColorFactor);
 
-	if (JsonMaterial.ShadingModel == EGLTFJsonShadingModel::Default || JsonMaterial.ShadingModel == EGLTFJsonShadingModel::ClearCoat)
+	if (JsonMaterial.ShadingModel == EGLTFJsonShadingModel::Default || 
+		JsonMaterial.ShadingModel == EGLTFJsonShadingModel::ClearCoat ||
+		JsonMaterial.ShadingModel == EGLTFJsonShadingModel::Sheen)
 	{
 		SetProxyParameter(ProxyMaterial, FGLTFProxyMaterialInfo::Emissive, JsonMaterial.EmissiveTexture);
 		SetProxyParameter(ProxyMaterial, FGLTFProxyMaterialInfo::EmissiveFactor, JsonMaterial.EmissiveFactor);
@@ -129,6 +131,13 @@ void FGLTFMaterialProxyFactory::SetProxyParameters(UMaterialInstanceConstant* Pr
 			SetProxyParameter(ProxyMaterial, FGLTFProxyMaterialInfo::ClearCoatNormal, JsonMaterial.ClearCoat.ClearCoatNormalTexture);
 			SetProxyParameter(ProxyMaterial, FGLTFProxyMaterialInfo::ClearCoatNormalScale, JsonMaterial.ClearCoat.ClearCoatNormalTexture.Scale);
 		}
+		else if (JsonMaterial.ShadingModel == EGLTFJsonShadingModel::Sheen)
+		{
+			SetProxyParameter(ProxyMaterial, FGLTFProxyMaterialInfo::SheenColorFactor, JsonMaterial.Sheen.ColorFactor);
+			SetProxyParameter(ProxyMaterial, FGLTFProxyMaterialInfo::SheenColorTexture, JsonMaterial.Sheen.ColorTexture);
+			SetProxyParameter(ProxyMaterial, FGLTFProxyMaterialInfo::SheenRoughnessFactor, JsonMaterial.Sheen.RoughnessFactor);
+			SetProxyParameter(ProxyMaterial, FGLTFProxyMaterialInfo::SheenRoughnessTexture, JsonMaterial.Sheen.RoughnessTexture);
+		}
 	}
 
 	if (JsonMaterial.ShadingModel != EGLTFJsonShadingModel::Unlit) //&& != SpecularGlossiness
diff --git a/Engine/Plugins/Enterprise/GLTFExporter/Source/GLTFExporter/Private/Json/GLTFJsonMaterial.cpp b/Engine/Plugins/Enterprise/GLTFExporter/Source/GLTFExporter/Private/Json/GLTFJsonMaterial.cpp
--- a/Engine/Plugins/Enterprise/GLTFExporter/Source/GLTFExporter/Private/Json/GLTFJsonMaterial.cpp
+++ b/Engine/Plugins/Enterprise/GLTFExporter/Source/GLTFExporter/Private/Json/GLTFJsonMaterial.cpp
@@ -141,6 +141,29 @@ void FGLTFJsonIORExtension::WriteObject(IGLTFJsonWriter& Writer) const
 	}
 }
 
+void FGLTFJsonSheenExtension::WriteObject(IGLTFJsonWriter& Writer) const
+{
+	if (!ColorFactor.IsNearlyEqual(FGLTFJsonColor3::Black))
+	{
+		Writer.Write(TEXT("sheenColorFactor"), ColorFactor);
+	}
+
+	if (ColorTexture.Index != nullptr)
+	{
+		Writer.Write(TEXT("sheenColorTexture"), ColorTexture);
+	}
+
+	if (!FMath::IsNearlyEqual(RoughnessFactor, 0.f))
+	{
+		Writer.Write(TEXT("sheenRoughnessFactor"), RoughnessFactor);
+	}
+
+	if (RoughnessTexture.Index != nullptr)
+	{
+		Writer.Write(TEXT("sheenRoughnessTexture"), RoughnessTexture);
+	}
+}
+
 void FGLTFJsonMaterial::WriteObject(IGLTFJsonWriter& Writer) const
 {
 	if (!Name.IsEmpty())
@@ -189,7 +212,11 @@ void FGLTFJsonMaterial::WriteObject(IGLTFJsonWriter& Writer) const
 	}
 
 	const bool HasEmissiveStrength = !FMath::IsNearlyEqual(EmissiveStrength, 1.0f, Writer.DefaultTolerance);
-	if (ShadingModel == EGLTFJsonShadingModel::Unlit || ShadingModel == EGLTFJsonShadingModel::ClearCoat || HasEmissiveStrength || Specular.HasValue() || IOR.HasValue())
+	if (ShadingModel == EGLTFJsonShadingModel::Unlit || ShadingModel == EGLTFJsonShadingModel::ClearCoat || 
+		HasEmissiveStrength || 
+		Specular.HasValue() || 
+		IOR.HasValue() ||
+		Sheen.HasValue())
 	{
 		Writer.StartExtensions();
 
@@ -203,6 +230,11 @@ void FGLTFJsonMaterial::WriteObject(IGLTFJsonWriter& Writer) const
 			Writer.Write(EGLTFJsonExtension::KHR_MaterialsIOR, IOR);
 		}
 
+		if (ShadingModel != EGLTFJsonShadingModel::Unlit /*&& !KHR_materials_pbrSpecularGlossiness*/ && Sheen.HasValue())
+		{
+			Writer.Write(EGLTFJsonExtension::KHR_MaterialsSheen, Sheen);
+		}
+
 		if (ShadingModel == EGLTFJsonShadingModel::Unlit)
 		{
 			// Write empty object
diff --git a/Engine/Plugins/Enterprise/GLTFExporter/Source/GLTFExporter/Private/Materials/GLTFProxyMaterialInfo.cpp b/Engine/Plugins/Enterprise/GLTFExporter/Source/GLTFExporter/Private/Materials/GLTFProxyMaterialInfo.cpp
--- a/Engine/Plugins/Enterprise/GLTFExporter/Source/GLTFExporter/Private/Materials/GLTFProxyMaterialInfo.cpp
+++ b/Engine/Plugins/Enterprise/GLTFExporter/Source/GLTFExporter/Private/Materials/GLTFProxyMaterialInfo.cpp
@@ -30,3 +30,8 @@ const TGLTFProxyMaterialParameterInfo<float> FGLTFProxyMaterialInfo::ClearCoatNo
 
 const TGLTFProxyMaterialParameterInfo<float> FGLTFProxyMaterialInfo::SpecularFactor = { TEXT("Specular Factor") };
 const FGLTFProxyMaterialTextureParameterInfo FGLTFProxyMaterialInfo::SpecularTexture = { TEXT("Specular") };
+
+const TGLTFProxyMaterialParameterInfo<FLinearColor> FGLTFProxyMaterialInfo::SheenColorFactor = { TEXT("Sheen Color Factor") };
+const FGLTFProxyMaterialTextureParameterInfo FGLTFProxyMaterialInfo::SheenColorTexture = { TEXT("Sheen Color") };
+const TGLTFProxyMaterialParameterInfo<float> FGLTFProxyMaterialInfo::SheenRoughnessFactor = { TEXT("Sheen Roughness Factor") };
+const FGLTFProxyMaterialTextureParameterInfo FGLTFProxyMaterialInfo::SheenRoughnessTexture = { TEXT("Sheen Roughness") };
diff --git a/Engine/Plugins/Enterprise/GLTFExporter/Source/GLTFExporter/Private/Materials/GLTFProxyMaterialInfo.h b/Engine/Plugins/Enterprise/GLTFExporter/Source/GLTFExporter/Private/Materials/GLTFProxyMaterialInfo.h
--- a/Engine/Plugins/Enterprise/GLTFExporter/Source/GLTFExporter/Private/Materials/GLTFProxyMaterialInfo.h
+++ b/Engine/Plugins/Enterprise/GLTFExporter/Source/GLTFExporter/Private/Materials/GLTFProxyMaterialInfo.h
@@ -36,4 +36,9 @@ public:
 	
 	static const TGLTFProxyMaterialParameterInfo<float> SpecularFactor;
 	static const FGLTFProxyMaterialTextureParameterInfo SpecularTexture; //Only using Alpha Channel
+
+	static const TGLTFProxyMaterialParameterInfo<FLinearColor> SheenColorFactor;
+	static const FGLTFProxyMaterialTextureParameterInfo        SheenColorTexture; //RGB
+	static const TGLTFProxyMaterialParameterInfo<float>        SheenRoughnessFactor;
+	static const FGLTFProxyMaterialTextureParameterInfo        SheenRoughnessTexture; //A
 };
diff --git a/Engine/Plugins/Enterprise/GLTFExporter/Source/GLTFExporter/Private/Options/GLTFExportOptions.cpp b/Engine/Plugins/Enterprise/GLTFExporter/Source/GLTFExporter/Private/Options/GLTFExportOptions.cpp
--- a/Engine/Plugins/Enterprise/GLTFExporter/Source/GLTFExporter/Private/Options/GLTFExportOptions.cpp
+++ b/Engine/Plugins/Enterprise/GLTFExporter/Source/GLTFExporter/Private/Options/GLTFExportOptions.cpp
@@ -17,6 +17,7 @@ void UGLTFExportOptions::ResetToDefault()
 	bExportProxyMaterials = true;
 	bExportUnlitMaterials = true;
 	bExportClearCoatMaterials = true;
+	bExportClothMaterials = true;
 	bExportEmissiveStrength = true;
 	BakeMaterialInputs = EGLTFMaterialBakeMode::UseMeshData;
 	DefaultMaterialBakeSize = EGLTFMaterialBakeSizePOT::POT_1024;
diff --git a/Engine/Plugins/Enterprise/GLTFExporter/Source/GLTFExporter/Private/Tasks/GLTFDelayedMaterialTasks.cpp b/Engine/Plugins/Enterprise/GLTFExporter/Source/GLTFExporter/Private/Tasks/GLTFDelayedMaterialTasks.cpp
--- a/Engine/Plugins/Enterprise/GLTFExporter/Source/GLTFExporter/Private/Tasks/GLTFDelayedMaterialTasks.cpp
+++ b/Engine/Plugins/Enterprise/GLTFExporter/Source/GLTFExporter/Private/Tasks/GLTFDelayedMaterialTasks.cpp
@@ -46,7 +46,6 @@ namespace
 	const TArray<FLinearColor> OcclusionInputMasks = { OcclusionMask };
 	const TArray<FLinearColor> ClearCoatInputMasks = { ClearCoatMask };
 	const TArray<FLinearColor> ClearCoatRoughnessInputMasks = { ClearCoatRoughnessMask };
-	const TArray<FLinearColor> SpecularMasks = { RedMask, GreenMask, BlueMask, AlphaMask };
 }
 
 void FGLTFDelayedMaterialTask::Process()
@@ -101,7 +100,9 @@ void FGLTFDelayedMaterialTask::Process()
 			}
 		}
 
-		if (JsonMaterial->ShadingModel == EGLTFJsonShadingModel::Default || JsonMaterial->ShadingModel == EGLTFJsonShadingModel::ClearCoat)
+		if (JsonMaterial->ShadingModel == EGLTFJsonShadingModel::Default || 
+			JsonMaterial->ShadingModel == EGLTFJsonShadingModel::ClearCoat ||
+			JsonMaterial->ShadingModel == EGLTFJsonShadingModel::Sheen)
 		{
 			const FMaterialPropertyEx MetallicProperty = MP_Metallic;
 			const FMaterialPropertyEx RoughnessProperty = MP_Roughness;
@@ -187,6 +188,17 @@ void FGLTFDelayedMaterialTask::Process()
 					}
 				}
 			}
+			else if (JsonMaterial->ShadingModel == EGLTFJsonShadingModel::Sheen)
+			{
+				const FMaterialPropertyEx FuzzColorProperty = MP_SubsurfaceColor;
+				const FMaterialPropertyEx ClothProperty = MP_CustomData0;
+
+				// We are not checking (IsPropertyNonDefault(FuzzColorProperty) || IsPropertyNonDefault(ClothProperty)) because UE's defaults are white/1 respectively compared to glTF black/0
+				if (!TryGetFuzzColorAndCloth(*JsonMaterial, FuzzColorProperty, ClothProperty))
+				{
+					Builder.LogWarning(FString::Printf(TEXT("Failed to export %s and %s for material %s"), *FuzzColorProperty.ToString(), *ClothProperty.ToString(), *Material->GetName()));
+				}
+			}
 		}
 	}
 
@@ -242,7 +254,9 @@ void FGLTFDelayedMaterialTask::GetProxyParameters(FGLTFJsonMaterial& OutMaterial
 	GetProxyParameter(FGLTFProxyMaterialInfo::BaseColor, OutMaterial.PBRMetallicRoughness.BaseColorTexture);
 	GetProxyParameter(FGLTFProxyMaterialInfo::BaseColorFactor, OutMaterial.PBRMetallicRoughness.BaseColorFactor);
 
-	if (OutMaterial.ShadingModel == EGLTFJsonShadingModel::Default || OutMaterial.ShadingModel == EGLTFJsonShadingModel::ClearCoat)
+	if (OutMaterial.ShadingModel == EGLTFJsonShadingModel::Default || 
+		OutMaterial.ShadingModel == EGLTFJsonShadingModel::ClearCoat ||
+		OutMaterial.ShadingModel == EGLTFJsonShadingModel::Sheen)
 	{
 		GetProxyParameter(FGLTFProxyMaterialInfo::Emissive, OutMaterial.EmissiveTexture);
 		GetProxyParameter(FGLTFProxyMaterialInfo::EmissiveFactor, OutMaterial.EmissiveFactor);
@@ -281,6 +295,13 @@ void FGLTFDelayedMaterialTask::GetProxyParameters(FGLTFJsonMaterial& OutMaterial
 			GetProxyParameter(FGLTFProxyMaterialInfo::ClearCoatNormal, OutMaterial.ClearCoat.ClearCoatNormalTexture);
 			GetProxyParameter(FGLTFProxyMaterialInfo::ClearCoatNormalScale, OutMaterial.ClearCoat.ClearCoatNormalTexture.Scale);
 		}
+		else if (OutMaterial.ShadingModel == EGLTFJsonShadingModel::Sheen)
+		{
+			GetProxyParameter(FGLTFProxyMaterialInfo::SheenColorFactor, OutMaterial.Sheen.ColorFactor);
+			GetProxyParameter(FGLTFProxyMaterialInfo::SheenColorTexture , OutMaterial.Sheen.ColorTexture);
+			GetProxyParameter(FGLTFProxyMaterialInfo::SheenRoughnessFactor, OutMaterial.Sheen.RoughnessFactor);
+			GetProxyParameter(FGLTFProxyMaterialInfo::SheenRoughnessTexture, OutMaterial.Sheen.RoughnessTexture);
+		}
 	}
 
 	if (OutMaterial.ShadingModel != EGLTFJsonShadingModel::Unlit)
@@ -478,6 +499,17 @@ void FGLTFDelayedMaterialTask::ConvertShadingModel(EGLTFJsonShadingModel& OutSha
 			*Material->GetName(),
 			*FGLTFNameUtilities::GetName(MSM_DefaultLit)));
 	}
+
+	if (OutShadingModel == EGLTFJsonShadingModel::Sheen && !Builder.ExportOptions->bExportClothMaterials)
+	{
+		OutShadingModel = EGLTFJsonShadingModel::Default;
+
+		Builder.LogWarning(FString::Printf(
+			TEXT("Shading model (%s) in material %s disabled by export options, will export as %s"),
+			*FGLTFNameUtilities::GetName(ShadingModel),
+			*Material->GetName(),
+			*FGLTFNameUtilities::GetName(MSM_DefaultLit)));
+	}
 }
 
 void FGLTFDelayedMaterialTask::ConvertAlphaMode(EGLTFJsonAlphaMode& OutAlphaMode) const
@@ -956,122 +988,12 @@ bool FGLTFDelayedMaterialTask::TryGetSpecular(FGLTFJsonMaterial& OutMaterial, co
 	}
 	else
 	{
-		if (Builder.ExportOptions->BakeMaterialInputs == EGLTFMaterialBakeMode::Disabled)
+		if (!TryGetBakedMaterialPropertyOntoAlphaChannel(OutMaterial.Specular.Texture, OutMaterial.Specular.Factor, 1.0f, SpecularProperty))
 		{
-			Builder.LogWarning(FString::Printf(
-				TEXT("%s for material %s needs to bake, but material baking is disabled by export options"),
-				*SpecularProperty.ToString(),
-				*Material->GetName()));
 			return false;
 		}
-
-		FGLTFPropertyBakeOutput PropertyBakeOutput = BakeMaterialProperty(SpecularProperty, OutMaterial.Specular.Texture.TexCoord, OutMaterial.Specular.Texture.Transform);
-
-		if (PropertyBakeOutput.bIsConstant)
-		{
-			OutMaterial.Specular.Factor = PropertyBakeOutput.ConstantValue.R;
-		}
-		else
-		{
-			if (Builder.ExportOptions->TextureImageFormat == EGLTFTextureImageFormat::None)
-			{
-				OutMaterial.Specular.Texture.Index = nullptr;
-				return true;
-			}
-
-			//Move Value to Alpha per documentation:
-			//
-			{
-				const FExpressionInput* MaterialInput = FGLTFMaterialUtilities::GetInputForProperty(Material, SpecularProperty);
-				if (MaterialInput == nullptr)
-				{
-					// TODO: report error
-					return false;
-				}
-				const FLinearColor Mask = FGLTFMaterialUtilities::GetMask(*MaterialInput);
-
-				enum ERGBMask
-				{
-					RED = 0,
-					GREEN = 1,
-					BLUE = 2,
-					ALPHA = 3
-				};
-
-				ERGBMask IndexOfMask = ERGBMask(0);
-
-				if (SpecularMasks.Contains(Mask))
-				{
-					IndexOfMask = ERGBMask(SpecularMasks.IndexOfByKey(Mask));
-				}
-
-				auto SwapValues = [](FColor& Color) {};
-
-				switch (IndexOfMask)
-				{
-					case RED:
-						for (FColor& Pixel : *PropertyBakeOutput.Pixels)
-						{
-							Pixel.A = Pixel.R;
-
-							Pixel.R = 255;
-							Pixel.G = 255;
-							Pixel.B = 255;
-						}
-						break;
-
-					case GREEN:
-						for (FColor& Pixel : *PropertyBakeOutput.Pixels)
-						{
-							Pixel.A = Pixel.G;
-
-							Pixel.R = 255;
-							Pixel.G = 255;
-							Pixel.B = 255;
-						}
-						break;
-
-					case BLUE:
-						for (FColor& Pixel : *PropertyBakeOutput.Pixels)
-						{
-							Pixel.A = Pixel.B;
-
-							Pixel.R = 255;
-							Pixel.G = 255;
-							Pixel.B = 255;
-						}
-						break;
-
-					case ALPHA:
-					default:
-						break;
-				}
-			}
-
-			const EGLTFMaterialPropertyGroup PropertyGroup = GetPropertyGroup(SpecularProperty);
-			const TextureAddress TextureAddress = Builder.GetBakeTilingForMaterialProperty(Material, PropertyGroup);
-			const TextureFilter TextureFilter = Builder.GetBakeFilterForMaterialProperty(Material, PropertyGroup);
-
-			FGLTFJsonTexture* Texture = FGLTFMaterialUtilities::AddTexture(
-				Builder,
-				PropertyBakeOutput.Pixels,
-				PropertyBakeOutput.Size,
-				false,
-				FGLTFMaterialUtilities::IsNormalMap(SpecularProperty),
-				GetBakedTextureName(TEXT("Specular")),
-				TextureAddress,
-				TextureFilter);
-
-			OutMaterial.Specular.Texture.Index = Texture;
-
-			if (!PropertyBakeOutput.bIsConstant)
-			{
-				OutMaterial.Specular.Factor = 1.0f;
-			}
-		}
 	}
 
-
 	return true;
 }
 
@@ -1125,6 +1047,43 @@ bool FGLTFDelayedMaterialTask::TryGetRefraction(FGLTFJsonMaterial& OutMaterial,
 	return true;
 }
 
+bool FGLTFDelayedMaterialTask::TryGetFuzzColorAndCloth(FGLTFJsonMaterial& OutMaterial, const FMaterialPropertyEx& FuzzColorProperty, const FMaterialPropertyEx& ClothProperty)
+{
+	//SheenColorFactor + SheenColorTexture(RGB)
+	if (!TryGetConstantColor(OutMaterial.Sheen.ColorFactor, FuzzColorProperty))
+	{
+		if (!TryGetSourceTexture(OutMaterial.Sheen.ColorTexture, FuzzColorProperty, { RgbMask }))
+		{
+			if (!TryGetBakedMaterialProperty(OutMaterial.Sheen.ColorTexture, OutMaterial.Sheen.ColorFactor, FuzzColorProperty, TEXT("SheenColorTexture")))
+			{
+				return false;
+			}
+		}
+		else
+		{
+			OutMaterial.Sheen.ColorFactor = FGLTFJsonColor3::White;
+		}
+	}
+
+	//SheenRoughnessFactor + SheenRoughnessTexture(A)
+	if (!TryGetConstantScalar(OutMaterial.Sheen.RoughnessFactor, ClothProperty))
+	{
+		if (!TryGetSourceTexture(OutMaterial.Sheen.RoughnessTexture, ClothProperty, { AlphaMask }))
+		{
+			if (!TryGetBakedMaterialPropertyOntoAlphaChannel(OutMaterial.Sheen.RoughnessTexture, OutMaterial.Sheen.RoughnessFactor, 1.0f, ClothProperty, TEXT("SheenRoughnessTexture")))
+			{
+				return false;
+			}
+		}
+		else
+		{
+			OutMaterial.Sheen.RoughnessFactor = 1.0f;
+		}
+	}
+
+	return true;
+}
+
 bool FGLTFDelayedMaterialTask::IsPropertyNonDefault(const FMaterialPropertyEx& Property) const
 {
 	// Custom outputs are by definition standalone and never part of material attributes
@@ -1695,6 +1654,129 @@ bool FGLTFDelayedMaterialTask::TryGetBakedMaterialProperty(FGLTFJsonTextureInfo&
 	return true;
 }
 
+bool FGLTFDelayedMaterialTask::TryGetBakedMaterialPropertyOntoAlphaChannel(FGLTFJsonTextureInfo& Texture, float& Factor, const float& DefaultFactorValue, const FMaterialPropertyEx& Property, const FString& TextureName)
+{
+	if (Builder.ExportOptions->BakeMaterialInputs == EGLTFMaterialBakeMode::Disabled)
+	{
+		Builder.LogWarning(FString::Printf(
+			TEXT("%s for material %s needs to bake, but material baking is disabled by export options"),
+			TextureName.IsEmpty() ? *Property.ToString() : *TextureName,
+			*Material->GetName()));
+		return false;
+	}
+
+	FGLTFPropertyBakeOutput PropertyBakeOutput = BakeMaterialProperty(Property, Texture.TexCoord, Texture.Transform);
+
+	if (PropertyBakeOutput.bIsConstant)
+	{
+		Factor = PropertyBakeOutput.ConstantValue.R;
+	}
+	else
+	{
+		if (Builder.ExportOptions->TextureImageFormat == EGLTFTextureImageFormat::None)
+		{
+			Texture.Index = nullptr;
+			return true;
+		}
+
+		//Move Value to Alpha per documentation:
+		//
+		{
+			const FExpressionInput* MaterialInput = FGLTFMaterialUtilities::GetInputForProperty(Material, Property);
+			if (MaterialInput == nullptr)
+			{
+				// TODO: report error
+				return false;
+			}
+			const FLinearColor Mask = FGLTFMaterialUtilities::GetMask(*MaterialInput);
+
+			enum ERGBMask
+			{
+				RED = 0,
+				GREEN = 1,
+				BLUE = 2,
+				ALPHA = 3,
+				RGB = 4
+			};
+
+			ERGBMask IndexOfMask = ERGBMask(0);
+
+			const TArray<FLinearColor> AllChannelMasks = { RedMask, GreenMask, BlueMask, AlphaMask, RgbMask };
+
+			if (AllChannelMasks.Contains(Mask))
+			{
+				IndexOfMask = ERGBMask(AllChannelMasks.IndexOfByKey(Mask));
+			}
+
+			auto SwapValues = [](FColor& Color) {};
+
+			switch (IndexOfMask)
+			{
+			case RED:
+			case RGB:
+				for (FColor& Pixel : *PropertyBakeOutput.Pixels)
+				{
+					Pixel.A = Pixel.R;
+
+					Pixel.R = 255;
+					Pixel.G = 255;
+					Pixel.B = 255;
+				}
+				break;
+
+			case GREEN:
+				for (FColor& Pixel : *PropertyBakeOutput.Pixels)
+				{
+					Pixel.A = Pixel.G;
+
+					Pixel.R = 255;
+					Pixel.G = 255;
+					Pixel.B = 255;
+				}
+				break;
+
+			case BLUE:
+				for (FColor& Pixel : *PropertyBakeOutput.Pixels)
+				{
+					Pixel.A = Pixel.B;
+
+					Pixel.R = 255;
+					Pixel.G = 255;
+					Pixel.B = 255;
+				}
+				break;
+
+			case ALPHA:
+			default:
+				break;
+			}
+		}
+
+		const EGLTFMaterialPropertyGroup PropertyGroup = GetPropertyGroup(Property);
+		const TextureAddress TextureAddress = Builder.GetBakeTilingForMaterialProperty(Material, PropertyGroup);
+		const TextureFilter TextureFilter = Builder.GetBakeFilterForMaterialProperty(Material, PropertyGroup);
+
+		FGLTFJsonTexture* JsonTexture = FGLTFMaterialUtilities::AddTexture(
+			Builder,
+			PropertyBakeOutput.Pixels,
+			PropertyBakeOutput.Size,
+			false,
+			FGLTFMaterialUtilities::IsNormalMap(Property),
+			GetBakedTextureName(TextureName.IsEmpty() ? Property.ToString() : TextureName),
+			TextureAddress,
+			TextureFilter);
+
+		Texture.Index = JsonTexture;
+
+		if (!PropertyBakeOutput.bIsConstant)
+		{
+			Factor = DefaultFactorValue;
+		}
+	}
+
+	return true;
+}
+
 FGLTFPropertyBakeOutput FGLTFDelayedMaterialTask::BakeMaterialProperty(const FMaterialPropertyEx& Property, int32& OutTexCoord, FGLTFJsonTextureTransform& OutTransform)
 {
 	const FIntPoint TextureSize = Builder.GetBakeSizeForMaterialProperty(Material, GetPropertyGroup(Property));
diff --git a/Engine/Plugins/Enterprise/GLTFExporter/Source/GLTFExporter/Private/Tasks/GLTFDelayedMaterialTasks.h b/Engine/Plugins/Enterprise/GLTFExporter/Source/GLTFExporter/Private/Tasks/GLTFDelayedMaterialTasks.h
--- a/Engine/Plugins/Enterprise/GLTFExporter/Source/GLTFExporter/Private/Tasks/GLTFDelayedMaterialTasks.h
+++ b/Engine/Plugins/Enterprise/GLTFExporter/Source/GLTFExporter/Private/Tasks/GLTFDelayedMaterialTasks.h
@@ -77,6 +77,7 @@ private:
 	bool TryGetEmissive(FGLTFJsonMaterial& OutMaterial, const FMaterialPropertyEx& EmissiveProperty);
 	bool TryGetSpecular(FGLTFJsonMaterial& OutMaterial, const FMaterialPropertyEx& SpecularProperty);
 	bool TryGetRefraction(FGLTFJsonMaterial& OutMaterial, const FMaterialPropertyEx& RefractionProperty);
+	bool TryGetFuzzColorAndCloth(FGLTFJsonMaterial& OutMaterial, const FMaterialPropertyEx& FuzzColorProperty, const FMaterialPropertyEx& ClothProperty);
 
 	bool IsPropertyNonDefault(const FMaterialPropertyEx& Property) const;
 	bool TryGetConstantColor(FGLTFJsonColor3& OutValue, const FMaterialPropertyEx& Property) const;
@@ -91,6 +92,7 @@ private:
 	bool TryGetBakedMaterialProperty(FGLTFJsonTextureInfo& OutTexInfo, FGLTFJsonColor4& OutConstant, const FMaterialPropertyEx& Property, const FString& PropertyName);
 	bool TryGetBakedMaterialProperty(FGLTFJsonTextureInfo& OutTexInfo, float& OutConstant, const FMaterialPropertyEx& Property, const FString& PropertyName);
 	bool TryGetBakedMaterialProperty(FGLTFJsonTextureInfo& OutTexInfo, const FMaterialPropertyEx& Property, const FString& PropertyName);
+	bool TryGetBakedMaterialPropertyOntoAlphaChannel(FGLTFJsonTextureInfo& Texture, float& Factor, const float& DefaultFactorValue, const FMaterialPropertyEx& Property, const FString& TextureName = TEXT(""));
 
 	FGLTFPropertyBakeOutput BakeMaterialProperty(const FMaterialPropertyEx& Property, int32& OutTexCoord, FGLTFJsonTextureTransform& OutTransform);
 	FGLTFPropertyBakeOutput BakeMaterialProperty(const FMaterialPropertyEx& Property, int32& OutTexCoord, FGLTFJsonTextureTransform& OutTransform, const FIntPoint& TextureSize, bool bFillAlpha);
diff --git a/Engine/Plugins/Enterprise/GLTFExporter/Source/GLTFExporter/Private/Utilities/GLTFCoreUtilities.cpp b/Engine/Plugins/Enterprise/GLTFExporter/Source/GLTFExporter/Private/Utilities/GLTFCoreUtilities.cpp
--- a/Engine/Plugins/Enterprise/GLTFExporter/Source/GLTFExporter/Private/Utilities/GLTFCoreUtilities.cpp
+++ b/Engine/Plugins/Enterprise/GLTFExporter/Source/GLTFExporter/Private/Utilities/GLTFCoreUtilities.cpp
@@ -222,6 +222,7 @@ EGLTFJsonShadingModel FGLTFCoreUtilities::ConvertShadingModel(EMaterialShadingMo
 		case MSM_Unlit:      return EGLTFJsonShadingModel::Unlit;
 		case MSM_DefaultLit: return EGLTFJsonShadingModel::Default;
 		case MSM_ClearCoat:  return EGLTFJsonShadingModel::ClearCoat;
+		case MSM_Cloth:      return EGLTFJsonShadingModel::Sheen;
 		default:             return EGLTFJsonShadingModel::None;
 	}
 }
diff --git a/Engine/Plugins/Enterprise/GLTFExporter/Source/GLTFExporter/Private/Utilities/GLTFJsonUtilities.cpp b/Engine/Plugins/Enterprise/GLTFExporter/Source/GLTFExporter/Private/Utilities/GLTFJsonUtilities.cpp
--- a/Engine/Plugins/Enterprise/GLTFExporter/Source/GLTFExporter/Private/Utilities/GLTFJsonUtilities.cpp
+++ b/Engine/Plugins/Enterprise/GLTFExporter/Source/GLTFExporter/Private/Utilities/GLTFJsonUtilities.cpp
@@ -12,6 +12,7 @@ const TCHAR* FGLTFJsonUtilities::GetValue(EGLTFJsonExtension Enum)
 		case EGLTFJsonExtension::KHR_MaterialsUnlit:            return TEXT("KHR_materials_unlit");
 		case EGLTFJsonExtension::KHR_MaterialsVariants:         return TEXT("KHR_materials_variants");
 		case EGLTFJsonExtension::KHR_MaterialsIOR:              return TEXT("KHR_materials_ior");
+		case EGLTFJsonExtension::KHR_MaterialsSheen:            return TEXT("KHR_materials_sheen");
 		case EGLTFJsonExtension::KHR_MeshQuantization:          return TEXT("KHR_mesh_quantization");
 		case EGLTFJsonExtension::KHR_TextureTransform:          return TEXT("KHR_texture_transform");
 		case EGLTFJsonExtension::KHR_MaterialsSpecular:	        return TEXT("KHR_materials_specular");
@@ -122,6 +123,7 @@ const TCHAR* FGLTFJsonUtilities::GetValue(EGLTFJsonShadingModel Enum)
 		case EGLTFJsonShadingModel::Default:   return TEXT("Default");
 		case EGLTFJsonShadingModel::Unlit:     return TEXT("Unlit");
 		case EGLTFJsonShadingModel::ClearCoat: return TEXT("ClearCoat");
+		case EGLTFJsonShadingModel::Sheen:     return TEXT("Sheen");
 		default:
 			checkNoEntry();
 			return TEXT("");
diff --git a/Engine/Plugins/Enterprise/GLTFExporter/Source/GLTFExporter/Public/Json/GLTFJsonEnums.h b/Engine/Plugins/Enterprise/GLTFExporter/Source/GLTFExporter/Public/Json/GLTFJsonEnums.h
--- a/Engine/Plugins/Enterprise/GLTFExporter/Source/GLTFExporter/Public/Json/GLTFJsonEnums.h
+++ b/Engine/Plugins/Enterprise/GLTFExporter/Source/GLTFExporter/Public/Json/GLTFJsonEnums.h
@@ -10,6 +10,7 @@ enum class EGLTFJsonExtension
 	KHR_MaterialsUnlit,
 	KHR_MaterialsVariants,
 	KHR_MaterialsIOR,
+	KHR_MaterialsSheen,
 	KHR_MeshQuantization,
 	KHR_TextureTransform,
 	KHR_MaterialsSpecular
@@ -21,6 +22,7 @@ enum class EGLTFJsonShadingModel
 	Default,
 	Unlit,
 	ClearCoat,
+	Sheen, /*UE Cloth*/
 
 	NumShadingModels
 };
diff --git a/Engine/Plugins/Enterprise/GLTFExporter/Source/GLTFExporter/Public/Json/GLTFJsonMaterial.h b/Engine/Plugins/Enterprise/GLTFExporter/Source/GLTFExporter/Public/Json/GLTFJsonMaterial.h
--- a/Engine/Plugins/Enterprise/GLTFExporter/Source/GLTFExporter/Public/Json/GLTFJsonMaterial.h
+++ b/Engine/Plugins/Enterprise/GLTFExporter/Source/GLTFExporter/Public/Json/GLTFJsonMaterial.h
@@ -116,6 +116,25 @@ struct GLTFEXPORTER_API FGLTFJsonIORExtension : IGLTFJsonObject
 	virtual void WriteObject(IGLTFJsonWriter& Writer) const override;
 };
 
+struct GLTFEXPORTER_API FGLTFJsonSheenExtension : IGLTFJsonObject
+{
+	FGLTFJsonColor3         ColorFactor;
+	FGLTFJsonTextureInfo    ColorTexture;
+
+	float                   RoughnessFactor;
+	FGLTFJsonTextureInfo    RoughnessTexture;
+
+	FGLTFJsonSheenExtension()
+		: ColorFactor(FGLTFJsonColor3::Black)
+		, RoughnessFactor(0)
+	{
+	}
+
+	bool HasValue() const { return !ColorFactor.IsNearlyEqual(FGLTFJsonColor3::Black) || ColorTexture.Index != nullptr || !FMath::IsNearlyEqual(RoughnessFactor, 0.f) || RoughnessTexture.Index != nullptr; }
+
+	virtual void WriteObject(IGLTFJsonWriter& Writer) const override;
+};
+
 struct GLTFEXPORTER_API FGLTFJsonMaterial : IGLTFJsonIndexedObject
 {
 	FString Name;
@@ -137,9 +156,9 @@ struct GLTFEXPORTER_API FGLTFJsonMaterial : IGLTFJsonIndexedObject
 	bool DoubleSided;
 
 	FGLTFJsonClearCoatExtension ClearCoat;
-	FGLTFJsonIORExtension IOR;
-
-	FGLTFJsonSpecularExtension Specular;
+	FGLTFJsonSpecularExtension  Specular;
+	FGLTFJsonIORExtension       IOR;
+	FGLTFJsonSheenExtension     Sheen;
 
 	virtual void WriteObject(IGLTFJsonWriter& Writer) const override;
 
diff --git a/Engine/Plugins/Enterprise/GLTFExporter/Source/GLTFExporter/Public/Options/GLTFExportOptions.h b/Engine/Plugins/Enterprise/GLTFExporter/Source/GLTFExporter/Public/Options/GLTFExportOptions.h
--- a/Engine/Plugins/Enterprise/GLTFExporter/Source/GLTFExporter/Public/Options/GLTFExportOptions.h
+++ b/Engine/Plugins/Enterprise/GLTFExporter/Source/GLTFExporter/Public/Options/GLTFExportOptions.h
@@ -61,6 +61,10 @@ class GLTFEXPORTER_API UGLTFExportOptions : public UObject
 	UPROPERTY(EditAnywhere, BlueprintReadWrite, Config, Category = Material)
 	bool bExportClearCoatMaterials;
 
+	/** If enabled, materials with shading model cloth will be properly exported. Uses extension KHR_materials_sheen. */
+	UPROPERTY(EditAnywhere, BlueprintReadWrite, Config, Category = Material)
+	bool bExportClothMaterials;
+
 	/** If enabled, allows materials to have an emissive factor that exceeds the standard range [0.0, 1.0]. Uses extension KHR_materials_emissive_strength. */
 	UPROPERTY(EditAnywhere, BlueprintReadWrite, Config, Category = Material)
 	bool bExportEmissiveStrength;
