Add support for defining arbitrary and non-square bake sizes and auto detecting bake size based on textures used in the material input graph

#rnx

diff --git a/Engine/Plugins/Enterprise/GLTFExporter/Source/GLTFExporter/GLTFExporter.Build.cs b/Engine/Plugins/Enterprise/GLTFExporter/Source/GLTFExporter/GLTFExporter.Build.cs
--- a/Engine/Plugins/Enterprise/GLTFExporter/Source/GLTFExporter/GLTFExporter.Build.cs
+++ b/Engine/Plugins/Enterprise/GLTFExporter/Source/GLTFExporter/GLTFExporter.Build.cs
@@ -56,6 +56,7 @@ public class GLTFExporter : ModuleRules
 					"EditorStyle",
 					"PropertyEditor",
 					"ToolMenus",
+					"ToolWidgets",
 					"ContentBrowser",
 					"MaterialUtilities",
 					"MeshMergeUtilities",
diff --git a/Engine/Plugins/Enterprise/GLTFExporter/Source/GLTFExporter/Private/Builders/GLTFBuilder.cpp b/Engine/Plugins/Enterprise/GLTFExporter/Source/GLTFExporter/Private/Builders/GLTFBuilder.cpp
--- a/Engine/Plugins/Enterprise/GLTFExporter/Source/GLTFExporter/Private/Builders/GLTFBuilder.cpp
+++ b/Engine/Plugins/Enterprise/GLTFExporter/Source/GLTFExporter/Private/Builders/GLTFBuilder.cpp
@@ -29,9 +29,9 @@ void FGLTFBuilder::ResolveProxies(TArray<const UMaterialInterface*>& Materials)
 	}
 }
 
-FIntPoint FGLTFBuilder::GetBakeSizeForMaterialProperty(const UMaterialInterface* Material, EGLTFMaterialPropertyGroup PropertyGroup) const
+FGLTFMaterialBakeSize FGLTFBuilder::GetBakeSizeForMaterialProperty(const UMaterialInterface* Material, EGLTFMaterialPropertyGroup PropertyGroup) const
 {
-	EGLTFMaterialBakeSizePOT DefaultValue = ExportOptions->DefaultMaterialBakeSize;
+	FGLTFMaterialBakeSize DefaultValue = ExportOptions->DefaultMaterialBakeSize;
 	if (const FGLTFOverrideMaterialBakeSettings* BakeSettings = ExportOptions->DefaultInputBakeSettings.Find(PropertyGroup))
 	{
 		if (BakeSettings->bOverrideSize)
@@ -41,9 +41,7 @@ FIntPoint FGLTFBuilder::GetBakeSizeForMaterialProperty(const UMaterialInterface*
 	}
 
 	// TODO: add option to ignore override
-	const EGLTFMaterialBakeSizePOT Size = UGLTFMaterialExportOptions::GetBakeSizeForPropertyGroup(Material, PropertyGroup, DefaultValue);
-	const int32 PixelSize = 1 << static_cast<uint8>(Size);
-	return { PixelSize, PixelSize };
+	return UGLTFMaterialExportOptions::GetBakeSizeForPropertyGroup(Material, PropertyGroup, DefaultValue);
 }
 
 TextureFilter FGLTFBuilder::GetBakeFilterForMaterialProperty(const UMaterialInterface* Material, EGLTFMaterialPropertyGroup PropertyGroup) const
diff --git a/Engine/Plugins/Enterprise/GLTFExporter/Source/GLTFExporter/Private/Converters/GLTFMaterialUtilities.cpp b/Engine/Plugins/Enterprise/GLTFExporter/Source/GLTFExporter/Private/Converters/GLTFMaterialUtilities.cpp
--- a/Engine/Plugins/Enterprise/GLTFExporter/Source/GLTFExporter/Private/Converters/GLTFMaterialUtilities.cpp
+++ b/Engine/Plugins/Enterprise/GLTFExporter/Source/GLTFExporter/Private/Converters/GLTFMaterialUtilities.cpp
@@ -6,17 +6,19 @@
 #include "Converters/GLTFProxyMaterialCompiler.h"
 #include "Utilities/GLTFProxyMaterialUtilities.h"
 #include "Materials/Material.h"
-#include "Misc/DefaultValueHelper.h"
+#include "Modules/ModuleManager.h"
+#include "Engine/RendererSettings.h"
+#include "Materials/MaterialExpressionTextureObject.h"
+#include "Materials/MaterialExpressionTextureObjectParameter.h"
 #if WITH_EDITOR
 #include "IMaterialBakingModule.h"
 #include "MaterialBakingStructures.h"
-#include "Materials/MaterialExpressionTextureSample.h"
 #include "Materials/MaterialAttributeDefinitionMap.h"
-#endif
-#include "Engine/RendererSettings.h"
-#include "Modules/ModuleManager.h"
 #include "Materials/MaterialExpressionCustomOutput.h"
 #include "Materials/MaterialExpressionTextureCoordinate.h"
+#include "Materials/MaterialExpressionTextureSample.h"
+#include "Materials/MaterialExpressionTextureSampleParameter2D.h"
+#endif
 
 UMaterialInterface* FGLTFMaterialUtilities::GetDefaultMaterial()
 {
@@ -31,6 +33,7 @@ bool FGLTFMaterialUtilities::IsClearCoatBottomNormalEnabled()
 
 #if WITH_EDITOR
 
+
 bool FGLTFMaterialUtilities::IsNormalMap(const FMaterialPropertyEx& Property)
 {
 	return Property == MP_Normal || Property == FMaterialPropertyEx::ClearCoatBottomNormal;
@@ -198,24 +201,117 @@ uint32 FGLTFMaterialUtilities::GetMaskComponentCount(const FExpressionInput& Exp
 	return ExpressionInput.MaskR + ExpressionInput.MaskG + ExpressionInput.MaskB + ExpressionInput.MaskA;
 }
 
-bool FGLTFMaterialUtilities::TryGetTextureCoordinateIndex(const UMaterialExpressionTextureSample* TextureSampler, int32& TexCoord, FGLTFJsonTextureTransform& Transform)
+bool FGLTFMaterialUtilities::TryGetMaxTextureSize(const UMaterialInterface* Material, const FMaterialPropertyEx& Property, FIntPoint& OutMaxSize)
+{
+	TArray<UMaterialExpressionTextureSample*> TextureSamples;
+	GetAllInputExpressionsOfType(Material, Property, TextureSamples);
+
+	if (TextureSamples.Num() == 0)
+	{
+		return false;
+	}
+
+	FIntPoint MaxSize = { 0, 0 };
+
+	for (const UMaterialExpressionTextureSample* TextureSample : TextureSamples)
+	{
+		const UTexture* Texture = GetTextureFromSample(Material, TextureSample);
+		if (Texture == nullptr || !FGLTFTextureUtilities::Is2D(Texture))
+		{
+			return false;
+		}
+
+		FGLTFTextureUtilities::FullyLoad(Texture);
+		FIntPoint TextureSize = FGLTFTextureUtilities::GetInGameSize(Texture);
+
+		MaxSize = MaxSize.ComponentMax(TextureSize);
+	}
+
+	OutMaxSize = MaxSize;
+	return true;
+}
+
+bool FGLTFMaterialUtilities::TryGetMaxTextureSize(const UMaterialInterface* Material, const FMaterialPropertyEx& PropertyA, const FMaterialPropertyEx& PropertyB, FIntPoint& OutMaxSize)
+{
+	FIntPoint MaxSizeA;
+	if (!TryGetMaxTextureSize(Material, PropertyA, MaxSizeA))
+	{
+		return false;
+	}
+
+	FIntPoint MaxSizeB;
+	if (!TryGetMaxTextureSize(Material, PropertyB, MaxSizeB))
+	{
+		return false;
+	}
+
+	OutMaxSize = MaxSizeA.ComponentMax(MaxSizeB);
+	return true;
+}
+
+UTexture* FGLTFMaterialUtilities::GetTextureFromSample(const UMaterialInterface* Material, const UMaterialExpressionTextureSample* SampleExpression)
+{
+	if (const UMaterialExpressionTextureSampleParameter2D* SampleParameter = ExactCast<UMaterialExpressionTextureSampleParameter2D>(SampleExpression))
+	{
+		UTexture* ParameterValue = SampleParameter->Texture;
+
+		if (!Material->GetTextureParameterValue(SampleParameter->GetParameterName(), ParameterValue))
+		{
+			return nullptr;
+		}
+
+		return ParameterValue;
+	}
+
+	if (const UMaterialExpressionTextureSample* Sample = ExactCast<UMaterialExpressionTextureSample>(SampleExpression))
+	{
+		UMaterialExpression* ObjectExpression = Sample->TextureObject.Expression;
+		if (ObjectExpression == nullptr)
+		{
+			return Sample->Texture;
+		}
+
+		if (const UMaterialExpressionTextureObjectParameter* ObjectParameter = ExactCast<UMaterialExpressionTextureObjectParameter>(ObjectExpression))
+		{
+			UTexture* ParameterValue = ObjectParameter->Texture;
+
+			if (!Material->GetTextureParameterValue(ObjectParameter->GetParameterName(), ParameterValue))
+			{
+				return nullptr;
+			}
+
+			return ParameterValue;
+		}
+
+		if (const UMaterialExpressionTextureObject* Object = ExactCast<UMaterialExpressionTextureObject>(ObjectExpression))
+		{
+			return Object->Texture;
+		}
+
+		return nullptr;
+	}
+
+	return nullptr;
+}
+
+bool FGLTFMaterialUtilities::TryGetTextureCoordinateIndex(const UMaterialExpressionTextureSample* TextureSample, int32& OutTexCoord, FGLTFJsonTextureTransform& OutTransform)
 {
-	const UMaterialExpression* Expression = TextureSampler->Coordinates.Expression;
+	const UMaterialExpression* Expression = TextureSample->Coordinates.Expression;
 	if (Expression == nullptr)
 	{
-		TexCoord = TextureSampler->ConstCoordinate;
-		Transform = {};
+		OutTexCoord = TextureSample->ConstCoordinate;
+		OutTransform = {};
 		return true;
 	}
 
 	if (const UMaterialExpressionTextureCoordinate* TextureCoordinate = Cast<UMaterialExpressionTextureCoordinate>(Expression))
 	{
-		TexCoord = TextureCoordinate->CoordinateIndex;
-		Transform.Offset.X = TextureCoordinate->UnMirrorU ? TextureCoordinate->UTiling * 0.5f : 0.0f;
-		Transform.Offset.Y = TextureCoordinate->UnMirrorV ? TextureCoordinate->VTiling * 0.5f : 0.0f;
-		Transform.Scale.X = TextureCoordinate->UTiling * (TextureCoordinate->UnMirrorU ? 0.5f : 1.0f);
-		Transform.Scale.Y = TextureCoordinate->VTiling * (TextureCoordinate->UnMirrorV ? 0.5f : 1.0f);
-		Transform.Rotation = 0;
+		OutTexCoord = TextureCoordinate->CoordinateIndex;
+		OutTransform.Offset.X = TextureCoordinate->UnMirrorU ? TextureCoordinate->UTiling * 0.5f : 0.0f;
+		OutTransform.Offset.Y = TextureCoordinate->UnMirrorV ? TextureCoordinate->VTiling * 0.5f : 0.0f;
+		OutTransform.Scale.X = TextureCoordinate->UTiling * (TextureCoordinate->UnMirrorU ? 0.5f : 1.0f);
+		OutTransform.Scale.Y = TextureCoordinate->VTiling * (TextureCoordinate->UnMirrorV ? 0.5f : 1.0f);
+		OutTransform.Rotation = 0;
 		return true;
 	}
 
@@ -224,10 +320,10 @@ bool FGLTFMaterialUtilities::TryGetTextureCoordinateIndex(const UMaterialExpress
 	return false;
 }
 
-void FGLTFMaterialUtilities::GetAllTextureCoordinateIndices(const UMaterialInterface* InMaterial, const FMaterialPropertyEx& InProperty, FGLTFIndexArray& OutTexCoords)
+void FGLTFMaterialUtilities::GetAllTextureCoordinateIndices(const UMaterialInterface* Material, const FMaterialPropertyEx& Property, FGLTFIndexArray& OutTexCoords)
 {
 	FMaterialAnalysisResult Analysis;
-	AnalyzeMaterialProperty(InMaterial, InProperty, Analysis);
+	AnalyzeMaterialProperty(Material, Property, Analysis);
 
 	const TBitArray<>& TexCoords = Analysis.TextureCoordinates;
 	for (int32 Index = 0; Index < TexCoords.Num(); Index++)
@@ -281,6 +377,56 @@ FMaterialShadingModelField FGLTFMaterialUtilities::EvaluateShadingModelExpressio
 	return Analysis.ShadingModels;
 }
 
+template <typename ExpressionType>
+void FGLTFMaterialUtilities::GetAllInputExpressionsOfType(const UMaterialInterface* Material, const FMaterialPropertyEx& Property, TArray<ExpressionType*>& OutExpressions)
+{
+	const FExpressionInput* Input = GetInputForProperty(Material, Property);
+	if (Input == nullptr)
+	{
+		return;
+	}
+
+	UMaterialExpression* InputExpression = Input->Expression;
+	if (InputExpression == nullptr)
+	{
+		return;
+	}
+
+	TArray<UMaterialExpression*> AllInputExpressions;
+	InputExpression->GetAllInputExpressions(AllInputExpressions);
+
+	for (UMaterialExpression* Expression : AllInputExpressions)
+	{
+		if (ExpressionType* ExpressionOfType = Cast<ExpressionType>(Expression))
+		{
+			OutExpressions.Add(ExpressionOfType);
+		}
+
+		if (UMaterialFunctionInterface* MaterialFunction = UMaterial::GetExpressionFunctionPointer(Expression))
+		{
+			MaterialFunction->GetAllExpressionsOfType<ExpressionType>(OutExpressions);
+		}
+		else if (TOptional<UMaterial::FLayersInterfaces> LayersInterfaces = UMaterial::GetExpressionLayers(Expression))
+		{
+			for (UMaterialFunctionInterface* Layer : LayersInterfaces->Layers)
+			{
+				if (Layer != nullptr)
+				{
+					Layer->GetAllExpressionsOfType<ExpressionType>(OutExpressions);
+				}
+			}
+
+			for (UMaterialFunctionInterface* Blend : LayersInterfaces->Blends)
+			{
+				if (Blend != nullptr)
+				{
+					Blend->GetAllExpressionsOfType<ExpressionType>(OutExpressions);
+				}
+			}
+		}
+	}
+}
+
 #endif
 
 EMaterialShadingModel FGLTFMaterialUtilities::GetRichestShadingModel(const FMaterialShadingModelField& ShadingModels)
diff --git a/Engine/Plugins/Enterprise/GLTFExporter/Source/GLTFExporter/Private/Converters/GLTFMaterialUtilities.h b/Engine/Plugins/Enterprise/GLTFExporter/Source/GLTFExporter/Private/Converters/GLTFMaterialUtilities.h
--- a/Engine/Plugins/Enterprise/GLTFExporter/Source/GLTFExporter/Private/Converters/GLTFMaterialUtilities.h
+++ b/Engine/Plugins/Enterprise/GLTFExporter/Source/GLTFExporter/Private/Converters/GLTFMaterialUtilities.h
@@ -41,6 +41,12 @@ struct FGLTFMaterialUtilities
 
 	static bool IsClearCoatBottomNormalEnabled();
 
+	static EMaterialShadingModel GetRichestShadingModel(const FMaterialShadingModelField& ShadingModels);
+	static FString ShadingModelsToString(const FMaterialShadingModelField& ShadingModels);
+
+	static bool NeedsMeshData(const UMaterialInterface* Material);
+	static bool NeedsMeshData(const TArray<const UMaterialInterface*>& Materials);
+
 #if WITH_EDITOR
 	static bool IsNormalMap(const FMaterialPropertyEx& Property);
 	static bool IsSRGB(const FMaterialPropertyEx& Property);
@@ -69,17 +75,21 @@ struct FGLTFMaterialUtilities
 	static FLinearColor GetMask(const FExpressionInput& ExpressionInput);
 	static uint32 GetMaskComponentCount(const FExpressionInput& ExpressionInput);
 
-	static bool TryGetTextureCoordinateIndex(const UMaterialExpressionTextureSample* TextureSampler, int32& TexCoord, FGLTFJsonTextureTransform& Transform);
-	static void GetAllTextureCoordinateIndices(const UMaterialInterface* InMaterial, const FMaterialPropertyEx& InProperty, FGLTFIndexArray& OutTexCoords);
+	static bool TryGetMaxTextureSize(const UMaterialInterface* Material, const FMaterialPropertyEx& Property, FIntPoint& OutMaxSize);
+	static bool TryGetMaxTextureSize(const UMaterialInterface* Material, const FMaterialPropertyEx& PropertyA, const FMaterialPropertyEx& PropertyB, FIntPoint& OutMaxSize);
+
+	static UTexture* GetTextureFromSample(const UMaterialInterface* Material, const UMaterialExpressionTextureSample* SampleExpression);
+
+	static bool TryGetTextureCoordinateIndex(const UMaterialExpressionTextureSample* TextureSample, int32& OutTexCoord, FGLTFJsonTextureTransform& OutTransform);
+	static void GetAllTextureCoordinateIndices(const UMaterialInterface* Material, const FMaterialPropertyEx& Property, FGLTFIndexArray& OutTexCoords);
 
 	static void AnalyzeMaterialProperty(const UMaterialInterface* Material, const FMaterialPropertyEx& InProperty, FMaterialAnalysisResult& OutAnalysis);
 
 	static FMaterialShadingModelField EvaluateShadingModelExpression(const UMaterialInterface* Material);
-#endif
 
-	static EMaterialShadingModel GetRichestShadingModel(const FMaterialShadingModelField& ShadingModels);
-	static FString ShadingModelsToString(const FMaterialShadingModelField& ShadingModels);
+private:
 
-	static bool NeedsMeshData(const UMaterialInterface* Material);
-	static bool NeedsMeshData(const TArray<const UMaterialInterface*>& Materials);
+	template<typename ExpressionType>
+	static void GetAllInputExpressionsOfType(const UMaterialInterface* Material, const FMaterialPropertyEx& Property, TArray<ExpressionType*>& OutExpressions);
+#endif
 };
diff --git a/Engine/Plugins/Enterprise/GLTFExporter/Source/GLTFExporter/Private/Converters/GLTFTextureUtilities.cpp b/Engine/Plugins/Enterprise/GLTFExporter/Source/GLTFExporter/Private/Converters/GLTFTextureUtilities.cpp
--- a/Engine/Plugins/Enterprise/GLTFExporter/Source/GLTFExporter/Private/Converters/GLTFTextureUtilities.cpp
+++ b/Engine/Plugins/Enterprise/GLTFExporter/Source/GLTFExporter/Private/Converters/GLTFTextureUtilities.cpp
@@ -47,6 +47,16 @@ void FGLTFTextureUtilities::FullyLoad(const UTexture* InTexture)
 	Texture->WaitForStreaming();
 }
 
+bool FGLTFTextureUtilities::Is2D(const UTexture* Texture)
+{
+	return Texture->IsA<UTexture2D>() || Texture->IsA<UTextureRenderTarget2D>();
+}
+
+bool FGLTFTextureUtilities::IsCubemap(const UTexture* Texture)
+{
+	return Texture->IsA<UTextureCube>() || Texture->IsA<UTextureRenderTargetCube>();
+}
+
 bool FGLTFTextureUtilities::IsHDR(const UTexture* Texture)
 {
 	switch (Texture->CompressionSettings)
@@ -60,11 +70,6 @@ bool FGLTFTextureUtilities::IsHDR(const UTexture* Texture)
 	}
 }
 
-bool FGLTFTextureUtilities::IsCubemap(const UTexture* Texture)
-{
-	return Texture->IsA<UTextureCube>() || Texture->IsA<UTextureRenderTargetCube>();
-}
-
 TextureFilter FGLTFTextureUtilities::GetDefaultFilter(TextureGroup LODGroup)
 {
 	const UTextureLODSettings* TextureLODSettings = UDeviceProfileManager::Get().GetActiveProfile()->GetTextureLODSettings();
@@ -93,8 +98,8 @@ int32 FGLTFTextureUtilities::GetMipBias(const UTexture* Texture)
 
 FIntPoint FGLTFTextureUtilities::GetInGameSize(const UTexture* Texture)
 {
-	const int32 Width = static_cast<int32>(Texture->GetSurfaceWidth());
-	const int32 Height = static_cast<int32>(Texture->GetSurfaceHeight());
+	const int32 Width = FMath::CeilToInt(Texture->GetSurfaceWidth());
+	const int32 Height = FMath::CeilToInt(Texture->GetSurfaceHeight());
 
 	const int32 MipBias = GetMipBias(Texture);
 
diff --git a/Engine/Plugins/Enterprise/GLTFExporter/Source/GLTFExporter/Private/Converters/GLTFTextureUtilities.h b/Engine/Plugins/Enterprise/GLTFExporter/Source/GLTFExporter/Private/Converters/GLTFTextureUtilities.h
--- a/Engine/Plugins/Enterprise/GLTFExporter/Source/GLTFExporter/Private/Converters/GLTFTextureUtilities.h
+++ b/Engine/Plugins/Enterprise/GLTFExporter/Source/GLTFExporter/Private/Converters/GLTFTextureUtilities.h
@@ -15,8 +15,9 @@ struct FGLTFTextureUtilities
 
 	static void FullyLoad(const UTexture* InTexture);
 
-	static bool IsHDR(const UTexture* Texture);
+	static bool Is2D(const UTexture* Texture);
 	static bool IsCubemap(const UTexture* Texture);
+	static bool IsHDR(const UTexture* Texture);
 
 	static TextureFilter GetDefaultFilter(TextureGroup Group);
 
diff --git a/Engine/Plugins/Enterprise/GLTFExporter/Source/GLTFExporter/Private/GLTFExporterModule.cpp b/Engine/Plugins/Enterprise/GLTFExporter/Source/GLTFExporter/Private/GLTFExporterModule.cpp
--- a/Engine/Plugins/Enterprise/GLTFExporter/Source/GLTFExporter/Private/GLTFExporterModule.cpp
+++ b/Engine/Plugins/Enterprise/GLTFExporter/Source/GLTFExporter/Private/GLTFExporterModule.cpp
@@ -1,21 +1,12 @@
 // Copyright Epic Games, Inc. All Rights Reserved.
 
 #include "GLTFExporterModule.h"
-#include "Json/GLTFJsonEnums.h"
-#include "Converters/GLTFObjectArrayScopeGuard.h"
-#include "Materials/Material.h"
 #include "Utilities/GLTFProxyMaterialUtilities.h"
+#include "UI/GLTFMaterialBakeSizeCustomization.h"
 #include "Actions/GLTFProxyAssetActions.h"
 #include "Interfaces/IPluginManager.h"
 #include "Misc/Paths.h"
 #include "ShaderCore.h"
-#if WITH_EDITOR
-#include "Materials/MaterialInstanceConstant.h"
-#include "Materials/MaterialInstanceDynamic.h"
-#include "AssetToolsModule.h"
-
-class FAssetTypeActions_Base;
-#endif
 
 DEFINE_LOG_CATEGORY(LogGLTFExporter);
 
@@ -31,6 +22,7 @@ public:
 
 #if WITH_EDITOR
 		// TODO: UI and editor-only functions should be moved (as much as possible) into their own module
+		FCoreDelegates::OnPostEngineInit.AddStatic(&FGLTFMaterialBakeSizeCustomization::Register);
 		FCoreDelegates::OnPostEngineInit.AddStatic(&FGLTFProxyAssetActions::AddMenuEntry);
 #endif
 	}
@@ -38,6 +30,7 @@ public:
 	virtual void ShutdownModule() override
 	{
 #if WITH_EDITOR
+		FGLTFMaterialBakeSizeCustomization::Unregister();
 		FGLTFProxyAssetActions::RemoveMenuEntry();
 #endif
 	}
diff --git a/Engine/Plugins/Enterprise/GLTFExporter/Source/GLTFExporter/Private/Options/GLTFExportOptions.cpp b/Engine/Plugins/Enterprise/GLTFExporter/Source/GLTFExporter/Private/Options/GLTFExportOptions.cpp
--- a/Engine/Plugins/Enterprise/GLTFExporter/Source/GLTFExporter/Private/Options/GLTFExportOptions.cpp
+++ b/Engine/Plugins/Enterprise/GLTFExporter/Source/GLTFExporter/Private/Options/GLTFExportOptions.cpp
@@ -20,7 +20,7 @@ void UGLTFExportOptions::ResetToDefault()
 	bExportClothMaterials = true;
 	bExportEmissiveStrength = true;
 	BakeMaterialInputs = EGLTFMaterialBakeMode::UseMeshData;
-	DefaultMaterialBakeSize = EGLTFMaterialBakeSizePOT::POT_1024;
+	DefaultMaterialBakeSize = FGLTFMaterialBakeSize::Default;
 	DefaultMaterialBakeFilter = TF_Trilinear;
 	DefaultMaterialBakeTiling = TA_Wrap;
 	DefaultLevelOfDetail = 0;
diff --git a/Engine/Plugins/Enterprise/GLTFExporter/Source/GLTFExporter/Private/Options/GLTFProxyOptions.cpp b/Engine/Plugins/Enterprise/GLTFExporter/Source/GLTFExporter/Private/Options/GLTFProxyOptions.cpp
--- a/Engine/Plugins/Enterprise/GLTFExporter/Source/GLTFExporter/Private/Options/GLTFProxyOptions.cpp
+++ b/Engine/Plugins/Enterprise/GLTFExporter/Source/GLTFExporter/Private/Options/GLTFProxyOptions.cpp
@@ -13,7 +13,7 @@ UGLTFProxyOptions::UGLTFProxyOptions(const FObjectInitializer& ObjectInitializer
 void UGLTFProxyOptions::ResetToDefault()
 {
 	bBakeMaterialInputs = true;
-	DefaultMaterialBakeSize = EGLTFMaterialBakeSizePOT::POT_1024;
+	DefaultMaterialBakeSize = FGLTFMaterialBakeSize::Default;
 	DefaultMaterialBakeFilter = TF_Trilinear;
 	DefaultMaterialBakeTiling = TA_Wrap;
 }
diff --git a/Engine/Plugins/Enterprise/GLTFExporter/Source/GLTFExporter/Private/Tasks/GLTFDelayedMaterialTasks.cpp b/Engine/Plugins/Enterprise/GLTFExporter/Source/GLTFExporter/Private/Tasks/GLTFDelayedMaterialTasks.cpp
--- a/Engine/Plugins/Enterprise/GLTFExporter/Source/GLTFExporter/Private/Tasks/GLTFDelayedMaterialTasks.cpp
+++ b/Engine/Plugins/Enterprise/GLTFExporter/Source/GLTFExporter/Private/Tasks/GLTFDelayedMaterialTasks.cpp
@@ -580,7 +580,7 @@ bool FGLTFDelayedMaterialTask::TryGetBaseColorAndOpacity(FGLTFJsonPBRMetallicRou
 		return false;
 	}
 
-	const FIntPoint TextureSize = Builder.GetBakeSizeForMaterialProperty(Material, GetPropertyGroup(BaseColorProperty));
+	const FIntPoint TextureSize = GetBakeSize(BaseColorProperty, OpacityProperty);
 	const TextureAddress TextureAddress = Builder.GetBakeTilingForMaterialProperty(Material, GetPropertyGroup(BaseColorProperty));
 	const TextureFilter TextureFilter = Builder.GetBakeFilterForMaterialProperty(Material, GetPropertyGroup(BaseColorProperty));
 
@@ -702,8 +702,7 @@ bool FGLTFDelayedMaterialTask::TryGetMetallicAndRoughness(FGLTFJsonPBRMetallicRo
 		return false;
 	}
 
-	// TODO: add support for calculating the ideal resolution to use for baking based on connected (texture) nodes
-	const FIntPoint TextureSize = Builder.GetBakeSizeForMaterialProperty(Material, GetPropertyGroup(MetallicProperty));
+	const FIntPoint TextureSize = GetBakeSize(MetallicProperty, RoughnessProperty);
 	const TextureAddress TextureAddress = Builder.GetBakeTilingForMaterialProperty(Material, GetPropertyGroup(MetallicProperty));
 	const TextureFilter TextureFilter = Builder.GetBakeFilterForMaterialProperty(Material, GetPropertyGroup(MetallicProperty));
 
@@ -822,7 +821,7 @@ bool FGLTFDelayedMaterialTask::TryGetClearCoatRoughness(FGLTFJsonClearCoatExtens
 		return false;
 	}
 
-	const FIntPoint TextureSize = Builder.GetBakeSizeForMaterialProperty(Material, GetPropertyGroup(IntensityProperty));
+	const FIntPoint TextureSize = GetBakeSize(IntensityProperty, RoughnessProperty);
 	const TextureAddress TextureAddress = Builder.GetBakeTilingForMaterialProperty(Material,GetPropertyGroup(IntensityProperty));
 	const TextureFilter TextureFilter = Builder.GetBakeFilterForMaterialProperty(Material, GetPropertyGroup(IntensityProperty));
 
@@ -1779,7 +1778,7 @@ bool FGLTFDelayedMaterialTask::TryGetBakedMaterialPropertyOntoAlphaChannel(FGLTF
 
 FGLTFPropertyBakeOutput FGLTFDelayedMaterialTask::BakeMaterialProperty(const FMaterialPropertyEx& Property, int32& OutTexCoord, FGLTFJsonTextureTransform& OutTransform)
 {
-	const FIntPoint TextureSize = Builder.GetBakeSizeForMaterialProperty(Material, GetPropertyGroup(Property));
+	const FIntPoint TextureSize = GetBakeSize(Property);
 	return BakeMaterialProperty(Property, OutTexCoord, OutTransform, TextureSize, true);
 }
 
@@ -1871,6 +1870,32 @@ bool FGLTFDelayedMaterialTask::StoreBakedPropertyTexture(const FMaterialProperty
 	return true;
 }
 
+FIntPoint FGLTFDelayedMaterialTask::GetBakeSize(const FMaterialPropertyEx& Property) const
+{
+	FGLTFMaterialBakeSize BakeSize = Builder.GetBakeSizeForMaterialProperty(Material, GetPropertyGroup(Property));
+
+	FIntPoint MaxSize;
+	if (BakeSize.bAutoDetect && FGLTFMaterialUtilities::TryGetMaxTextureSize(Material, Property, MaxSize))
+	{
+		return MaxSize;
+	}
+
+	return { BakeSize.X, BakeSize.Y };
+}
+
+FIntPoint FGLTFDelayedMaterialTask::GetBakeSize(const FMaterialPropertyEx& PropertyA, const FMaterialPropertyEx& PropertyB) const
+{
+	FGLTFMaterialBakeSize BakeSize = Builder.GetBakeSizeForMaterialProperty(Material, GetPropertyGroup(PropertyA));
+
+	FIntPoint MaxSize;
+	if (BakeSize.bAutoDetect && FGLTFMaterialUtilities::TryGetMaxTextureSize(Material, PropertyA, PropertyB, MaxSize))
+	{
+		return MaxSize;
+	}
+
+	return { BakeSize.X, BakeSize.Y };
+}
+
 EGLTFMaterialPropertyGroup FGLTFDelayedMaterialTask::GetPropertyGroup(const FMaterialPropertyEx& Property)
 {
 	switch (Property.Type)
diff --git a/Engine/Plugins/Enterprise/GLTFExporter/Source/GLTFExporter/Private/Tasks/GLTFDelayedMaterialTasks.h b/Engine/Plugins/Enterprise/GLTFExporter/Source/GLTFExporter/Private/Tasks/GLTFDelayedMaterialTasks.h
--- a/Engine/Plugins/Enterprise/GLTFExporter/Source/GLTFExporter/Private/Tasks/GLTFDelayedMaterialTasks.h
+++ b/Engine/Plugins/Enterprise/GLTFExporter/Source/GLTFExporter/Private/Tasks/GLTFDelayedMaterialTasks.h
@@ -99,6 +99,9 @@ private:
 
 	bool StoreBakedPropertyTexture(const FMaterialPropertyEx& Property, FGLTFJsonTextureInfo& OutTexInfo, FGLTFPropertyBakeOutput& PropertyBakeOutput, const FString& PropertyName) const;
 
+	FIntPoint GetBakeSize(const FMaterialPropertyEx& Property) const;
+	FIntPoint GetBakeSize(const FMaterialPropertyEx& PropertyA,const FMaterialPropertyEx& PropertyB) const;
+
 	static EGLTFMaterialPropertyGroup GetPropertyGroup(const FMaterialPropertyEx& Property);
 
 	template <typename CallbackType>
diff --git a/Engine/Plugins/Enterprise/GLTFExporter/Source/GLTFExporter/Private/UI/GLTFMaterialBakeSizeCustomization.cpp b/Engine/Plugins/Enterprise/GLTFExporter/Source/GLTFExporter/Private/UI/GLTFMaterialBakeSizeCustomization.cpp
new file mode 100644
--- /dev/null
+++ b/Engine/Plugins/Enterprise/GLTFExporter/Source/GLTFExporter/Private/UI/GLTFMaterialBakeSizeCustomization.cpp
@@ -0,0 +1,241 @@
+// Copyright Epic Games, Inc. All Rights Reserved.
+
+#if WITH_EDITOR
+
+#include "UI/GLTFMaterialBakeSizeCustomization.h"
+#include "UserData/GLTFMaterialUserData.h"
+#include "Widgets/Input/SCheckBox.h"
+#include "SSearchableComboBox.h"
+#include "DetailWidgetRow.h"
+#include "IPropertyUtilities.h"
+#include "Modules/ModuleManager.h"
+
+#define LOCTEXT_NAMESPACE "GLTFMaterialBakeSettingsDetails"
+
+void FGLTFMaterialBakeSizeCustomization::Register()
+{
+	static FName PropertyTypeName("GLTFMaterialBakeSize");
+	static FName PropertyEditor("PropertyEditor");
+	FPropertyEditorModule& PropertyModule = FModuleManager::GetModuleChecked<FPropertyEditorModule>(PropertyEditor);
+	PropertyModule.RegisterCustomPropertyTypeLayout(PropertyTypeName, FOnGetPropertyTypeCustomizationInstance::CreateStatic(&FGLTFMaterialBakeSizeCustomization::MakeInstance));
+}
+
+void FGLTFMaterialBakeSizeCustomization::Unregister()
+{
+	static FName PropertyTypeName("GLTFMaterialBakeSize");
+	static FName PropertyEditor("PropertyEditor");
+	FPropertyEditorModule& PropertyModule = FModuleManager::GetModuleChecked<FPropertyEditorModule>(PropertyEditor);
+	PropertyModule.UnregisterCustomPropertyTypeLayout(PropertyTypeName);
+}
+
+TSharedRef<IPropertyTypeCustomization> FGLTFMaterialBakeSizeCustomization::MakeInstance()
+{
+	return MakeShared<FGLTFMaterialBakeSizeCustomization>();
+}
+
+void FGLTFMaterialBakeSizeCustomization::CustomizeHeader(TSharedRef<IPropertyHandle> PropertyHandle, FDetailWidgetRow& HeaderRow, IPropertyTypeCustomizationUtils& CustomizationUtils)
+{
+	UE_LOG(LogTemp, Warning, TEXT("FGLTFMaterialBakeSettingsDetails: %s"), *FGLTFMaterialBakeSize::StaticStruct()->GetName());
+	StructPropertyHandle = PropertyHandle;
+
+	UStruct* Struct = FGLTFMaterialBakeSize::StaticStruct();
+	
+	FProperty* Property = Struct->FindPropertyByName(TEXT("bAutoDetect"));
+
+	FText AutoDetectToolTip = FText::FromString(Property->GetMetaData(TEXT("ToolTip")));
+	
+	TSharedPtr<IPropertyUtilities> PropertyUtils = CustomizationUtils.GetPropertyUtilities();
+
+	int32 Size = 1;
+	while (Size <= 8192)
+	{		
+		OptionsSource.Add(MakeShared<FString>(FString::Printf(TEXT("%d x %d"), Size, Size)));
+		Size *= 2;
+	}
+
+	FSlateFontInfo SmallFont = FCoreStyle::Get().GetFontStyle("SmallFont");
+
+	HeaderRow
+	.NameContent()
+		[
+			PropertyHandle->CreatePropertyNameWidget()
+		]
+	.ValueContent()
+		[
+			SNew(SHorizontalBox)
+			+SHorizontalBox::Slot()
+			.AutoWidth()
+			[
+				SNew(SSearchableComboBox)
+				//.ContentPadding(0)
+				.OptionsSource(&OptionsSource)
+				.OnGenerateWidget_Lambda([SmallFont](const TSharedPtr<FString>& Item)
+				{
+					return SNew(STextBlock)
+						.Text(FText::FromString(*Item))
+						.Font(SmallFont);
+				})
+				.OnSelectionChanged_Lambda([this](const TSharedPtr<FString>& SelectedItem, ESelectInfo::Type Type)
+				{
+					SetSizeString(*SelectedItem);
+				})
+				.InitiallySelectedItem(GetSelectedOption())
+				.SearchVisibility(EVisibility::Collapsed)
+				.Content()
+				[
+					SNew(SEditableText)
+					.MinDesiredWidth(89.0f) // TODO: instead of hardcoded value, should calculate it based on something
+					.RevertTextOnEscape(true)
+					.SelectAllTextWhenFocused(true)
+					.Text_Lambda([this]()
+					{
+						return FText::FromString(GetSizeString());
+					})
+					.OnTextCommitted_Lambda([this](const FText& NewText, ETextCommit::Type Type)
+                    {
+						SetSizeString(NewText.ToString());
+                    })
+					.Font(SmallFont)
+				]
+			]
+			+SHorizontalBox::Slot()
+			.Padding(10.0f, 0.0f, 0.0f, 0.0f)
+			.AutoWidth()
+			[
+				SNew(SCheckBox)
+				.IsChecked_Lambda([this]()
+				{
+					bool bAutoDetect;
+					return TryGetPropertyValue<bool>(bAutoDetect, [](const FGLTFMaterialBakeSize& Size){ return Size.bAutoDetect; })
+						? (bAutoDetect ? ECheckBoxState::Checked : ECheckBoxState::Unchecked) : ECheckBoxState::Undetermined;
+				})
+				.OnCheckStateChanged_Lambda([this](ECheckBoxState NewState)
+				{
+					bool bAutoDetect = NewState == ECheckBoxState::Checked;
+					SetPropertyValue([bAutoDetect](FGLTFMaterialBakeSize& Size){ Size.bAutoDetect = bAutoDetect; });
+				})
+				.ToolTipText(AutoDetectToolTip)
+				[
+					SNew(STextBlock)
+					.Margin(FMargin(3.0f, 0.0f, 0.0f, 0.0f))
+					.Text(LOCTEXT("AutoDetectLabel", "Auto-detect"))
+					.Font(SmallFont)
+				]
+			]
+		].IsEnabled(MakeAttributeLambda([=] { return !PropertyHandle->IsEditConst() && PropertyUtils->IsPropertyEditingEnabled(); }));
+}
+
+void FGLTFMaterialBakeSizeCustomization::CustomizeChildren(TSharedRef<IPropertyHandle> PropertyHandle, IDetailChildrenBuilder& ChildBuilder, IPropertyTypeCustomizationUtils& CustomizationUtils)
+{
+}
+
+FString FGLTFMaterialBakeSizeCustomization::GetSizeString() const
+{
+	FIntPoint Value;
+	if (!TryGetPropertyValue<FIntPoint>(Value, [](const FGLTFMaterialBakeSize& Size) { return FIntPoint(Size.X, Size.Y); }))
+	{
+		return TEXT("");
+	}
+
+	return FString::Printf(TEXT("%d x %d"), Value.X, Value.Y);
+}
+
+void FGLTFMaterialBakeSizeCustomization::SetSizeString(const FString& String) const
+{
+	auto ParseInt = [](const FString& Text, int32& OutValue, int32 DefaultValue)
+	{
+		const TCHAR* TextStart = *Text;
+		TCHAR* TextEnd;
+
+		OutValue = FCString::Strtoi(TextStart, &TextEnd, 10);
+		if (TextEnd == TextStart || OutValue <= 0)
+		{
+			OutValue = DefaultValue;
+		}
+	};
+
+	FIntPoint Value;
+
+	FString Left, Right;
+	if (String.Split(TEXT("x"), &Left, &Right))
+	{
+		ParseInt(*Left, Value.X, FGLTFMaterialBakeSize().X);
+		ParseInt(*Right, Value.Y, FGLTFMaterialBakeSize().Y);
+	}
+	else
+	{
+		ParseInt(*String, Value.X, FGLTFMaterialBakeSize().X);
+		Value.Y = Value.X;
+	}
+
+	SetPropertyValue([Value](FGLTFMaterialBakeSize& Size)
+	{
+		Size.X = Value.X;
+		Size.Y = Value.Y;
+	});
+}
+
+TSharedPtr<FString> FGLTFMaterialBakeSizeCustomization::GetSelectedOption() const
+{
+	const FString CurrentValue = GetSizeString();
+	const TSharedPtr<FString>* CurrentItem = OptionsSource.FindByPredicate([CurrentValue](const TSharedPtr<FString>& Item) { return *Item == CurrentValue; });
+	return CurrentItem != nullptr ? *CurrentItem : nullptr;
+}
+
+template <typename SetterType>
+void FGLTFMaterialBakeSizeCustomization::SetPropertyValue(SetterType Setter) const
+{
+	TArray<void*> RawPtrs;
+	StructPropertyHandle->AccessRawData(RawPtrs);
+	StructPropertyHandle->NotifyPreChange();
+
+	for (void* RawPtr : RawPtrs)
+	{
+		if (FGLTFMaterialBakeSize* Ptr = static_cast<FGLTFMaterialBakeSize*>(RawPtr))
+		{
+			Setter(*Ptr);
+		}
+	}
+
+	StructPropertyHandle->NotifyPostChange(EPropertyChangeType::ValueSet);
+	StructPropertyHandle->NotifyFinishedChangingProperties();
+}
+
+template <typename ValueType, typename GetterType>
+bool FGLTFMaterialBakeSizeCustomization::TryGetPropertyValue(ValueType& OutValue, GetterType Getter) const
+{
+	TArray<const void*> RawPtrs;
+	StructPropertyHandle->AccessRawData(RawPtrs);
+
+	if (RawPtrs.Num() == 0)
+	{
+		return false;
+	}
+
+	const FGLTFMaterialBakeSize* FirstPtr = static_cast<const FGLTFMaterialBakeSize*>(RawPtrs[0]);
+	if (FirstPtr == nullptr)
+	{
+		return false;
+	}
+
+	ValueType FirstValue = Getter(*FirstPtr);
+
+	for (int32 Index = 1; Index < RawPtrs.Num(); ++Index)
+	{
+		if (const FGLTFMaterialBakeSize* Ptr = static_cast<const FGLTFMaterialBakeSize*>(RawPtrs[Index]))
+		{
+			ValueType Value = Getter(*Ptr);
+			if (Value != FirstValue)
+			{
+				return false;
+			}
+		}
+	}
+
+	OutValue = FirstValue;
+	return true;
+}
+
+#undef LOCTEXT_NAMESPACE
+
+#endif
diff --git a/Engine/Plugins/Enterprise/GLTFExporter/Source/GLTFExporter/Private/UI/GLTFMaterialBakeSizeCustomization.h b/Engine/Plugins/Enterprise/GLTFExporter/Source/GLTFExporter/Private/UI/GLTFMaterialBakeSizeCustomization.h
new file mode 100644
--- /dev/null
+++ b/Engine/Plugins/Enterprise/GLTFExporter/Source/GLTFExporter/Private/UI/GLTFMaterialBakeSizeCustomization.h
@@ -0,0 +1,46 @@
+// Copyright Epic Games, Inc. All Rights Reserved.
+
+#pragma once
+
+#if WITH_EDITOR
+
+#include "Containers/Map.h"
+#include "HAL/PlatformCrt.h"
+#include "IPropertyTypeCustomization.h"
+#include "Templates/SharedPointer.h"
+#include "Types/SlateEnums.h"
+
+class FName;
+class FString;
+class IDetailLayoutBuilder;
+class UEnum;
+
+class FGLTFMaterialBakeSizeCustomization : public IPropertyTypeCustomization
+{
+public:
+
+	static void Register();
+	static void Unregister();
+
+	static TSharedRef<IPropertyTypeCustomization> MakeInstance();
+
+	virtual void CustomizeHeader(TSharedRef<IPropertyHandle> PropertyHandle, FDetailWidgetRow& HeaderRow, IPropertyTypeCustomizationUtils& CustomizationUtils) override;
+	virtual void CustomizeChildren(TSharedRef<IPropertyHandle> PropertyHandle, IDetailChildrenBuilder& ChildBuilder, IPropertyTypeCustomizationUtils& CustomizationUtils) override;
+
+private:
+
+	FString GetSizeString() const;
+	void SetSizeString(const FString& String) const;
+	TSharedPtr<FString> GetSelectedOption() const;
+
+	template <typename SetterType>
+	void SetPropertyValue(SetterType Setter) const;
+
+	template <typename ValueType, typename GetterType>
+	bool TryGetPropertyValue(ValueType& OutValue, GetterType Getter) const;
+
+	TSharedPtr<IPropertyHandle> StructPropertyHandle;
+	TArray<TSharedPtr<FString>> OptionsSource;
+};
+
+#endif
diff --git a/Engine/Plugins/Enterprise/GLTFExporter/Source/GLTFExporter/Private/UserData/GLTFMaterialUserData.cpp b/Engine/Plugins/Enterprise/GLTFExporter/Source/GLTFExporter/Private/UserData/GLTFMaterialUserData.cpp
--- a/Engine/Plugins/Enterprise/GLTFExporter/Source/GLTFExporter/Private/UserData/GLTFMaterialUserData.cpp
+++ b/Engine/Plugins/Enterprise/GLTFExporter/Source/GLTFExporter/Private/UserData/GLTFMaterialUserData.cpp
@@ -6,9 +6,11 @@
 
 #include UE_INLINE_GENERATED_CPP_BY_NAME(GLTFMaterialUserData)
 
+FGLTFMaterialBakeSize FGLTFMaterialBakeSize::Default = { 1024, 1024, true };
+
 FGLTFOverrideMaterialBakeSettings::FGLTFOverrideMaterialBakeSettings()
 	: bOverrideSize(false)
-	, Size(EGLTFMaterialBakeSizePOT::POT_1024)
+	, Size(FGLTFMaterialBakeSize::Default)
 	, bOverrideFilter(false)
 	, Filter(TF_Trilinear)
 	, bOverrideTiling(false)
@@ -50,7 +52,7 @@ const UMaterialInterface* UGLTFMaterialExportOptions::ResolveProxy(const UMateri
 	}
 }
 
-EGLTFMaterialBakeSizePOT UGLTFMaterialExportOptions::GetBakeSizeForPropertyGroup(const UMaterialInterface* Material, EGLTFMaterialPropertyGroup PropertyGroup, EGLTFMaterialBakeSizePOT DefaultValue)
+FGLTFMaterialBakeSize UGLTFMaterialExportOptions::GetBakeSizeForPropertyGroup(const UMaterialInterface* Material, EGLTFMaterialPropertyGroup PropertyGroup, FGLTFMaterialBakeSize DefaultValue)
 {
 	if (const FGLTFOverrideMaterialBakeSettings* BakeSettings = GetBakeSettingsByPredicate(Material, PropertyGroup, [](const FGLTFOverrideMaterialBakeSettings& Settings) { return Settings.bOverrideSize; }))
 	{
diff --git a/Engine/Plugins/Enterprise/GLTFExporter/Source/GLTFExporter/Public/Builders/GLTFBuilder.h b/Engine/Plugins/Enterprise/GLTFExporter/Source/GLTFExporter/Public/Builders/GLTFBuilder.h
--- a/Engine/Plugins/Enterprise/GLTFExporter/Source/GLTFExporter/Public/Builders/GLTFBuilder.h
+++ b/Engine/Plugins/Enterprise/GLTFExporter/Source/GLTFExporter/Public/Builders/GLTFBuilder.h
@@ -26,7 +26,7 @@ public:
 	const UMaterialInterface* ResolveProxy(const UMaterialInterface* Material) const;
 	void ResolveProxies(TArray<const UMaterialInterface*>& Materials) const;
 
-	FIntPoint GetBakeSizeForMaterialProperty(const UMaterialInterface* Material, EGLTFMaterialPropertyGroup PropertyGroup) const;
+	FGLTFMaterialBakeSize GetBakeSizeForMaterialProperty(const UMaterialInterface* Material, EGLTFMaterialPropertyGroup PropertyGroup) const;
 	TextureFilter GetBakeFilterForMaterialProperty(const UMaterialInterface* Material, EGLTFMaterialPropertyGroup PropertyGroup) const;
 	TextureAddress GetBakeTilingForMaterialProperty(const UMaterialInterface* Material, EGLTFMaterialPropertyGroup PropertyGroup) const;
 
diff --git a/Engine/Plugins/Enterprise/GLTFExporter/Source/GLTFExporter/Public/Options/GLTFExportOptions.h b/Engine/Plugins/Enterprise/GLTFExporter/Source/GLTFExporter/Public/Options/GLTFExportOptions.h
--- a/Engine/Plugins/Enterprise/GLTFExporter/Source/GLTFExporter/Public/Options/GLTFExportOptions.h
+++ b/Engine/Plugins/Enterprise/GLTFExporter/Source/GLTFExporter/Public/Options/GLTFExportOptions.h
@@ -75,7 +75,7 @@ class GLTFEXPORTER_API UGLTFExportOptions : public UObject
 
 	/** Default size of the baked out texture (containing the material input). Can be overridden by material- and input-specific bake settings, see GLTFMaterialExportOptions. */
 	UPROPERTY(EditAnywhere, BlueprintReadWrite, Config, Category = Material, Meta = (EditCondition = "BakeMaterialInputs != EGLTFMaterialBakeMode::Disabled"))
-	EGLTFMaterialBakeSizePOT DefaultMaterialBakeSize;
+	FGLTFMaterialBakeSize DefaultMaterialBakeSize;
 
 	/** Default filtering mode used when sampling the baked out texture. Can be overridden by material- and input-specific bake settings, see GLTFMaterialExportOptions. */
 	UPROPERTY(EditAnywhere, BlueprintReadWrite, Config, Category = Material, Meta = (EditCondition = "BakeMaterialInputs != EGLTFMaterialBakeMode::Disabled", ValidEnumValues = "TF_Nearest, TF_Bilinear, TF_Trilinear"))
diff --git a/Engine/Plugins/Enterprise/GLTFExporter/Source/GLTFExporter/Public/Options/GLTFProxyOptions.h b/Engine/Plugins/Enterprise/GLTFExporter/Source/GLTFExporter/Public/Options/GLTFProxyOptions.h
--- a/Engine/Plugins/Enterprise/GLTFExporter/Source/GLTFExporter/Public/Options/GLTFProxyOptions.h
+++ b/Engine/Plugins/Enterprise/GLTFExporter/Source/GLTFExporter/Public/Options/GLTFProxyOptions.h
@@ -16,7 +16,7 @@ class GLTFEXPORTER_API UGLTFProxyOptions : public UObject
 
 	/** Default size of the baked out texture (containing the material input). Can be overridden by material- and input-specific bake settings, see GLTFMaterialExportOptions. */
 	UPROPERTY(EditAnywhere, BlueprintReadWrite, Config, Category = Material, Meta = (EditCondition = "bBakeMaterialInputs"))
-	EGLTFMaterialBakeSizePOT DefaultMaterialBakeSize;
+	FGLTFMaterialBakeSize DefaultMaterialBakeSize;
 
 	/** Default filtering mode used when sampling the baked out texture. Can be overridden by material- and input-specific bake settings, see GLTFMaterialExportOptions. */
 	UPROPERTY(EditAnywhere, BlueprintReadWrite, Config, Category = Material, Meta = (EditCondition = "bBakeMaterialInputs", ValidEnumValues = "TF_Nearest, TF_Bilinear, TF_Trilinear"))
diff --git a/Engine/Plugins/Enterprise/GLTFExporter/Source/GLTFExporter/Public/UserData/GLTFMaterialUserData.h b/Engine/Plugins/Enterprise/GLTFExporter/Source/GLTFExporter/Public/UserData/GLTFMaterialUserData.h
--- a/Engine/Plugins/Enterprise/GLTFExporter/Source/GLTFExporter/Public/UserData/GLTFMaterialUserData.h
+++ b/Engine/Plugins/Enterprise/GLTFExporter/Source/GLTFExporter/Public/UserData/GLTFMaterialUserData.h
@@ -19,25 +19,6 @@ enum class EGLTFMaterialBakeMode : uint8
 	UseMeshData,
 };
 
-UENUM(BlueprintType)
-enum class EGLTFMaterialBakeSizePOT : uint8
-{
-	POT_1 UMETA(DisplayName = "1 x 1"),
-	POT_2 UMETA(DisplayName = "2 x 2"),
-	POT_4 UMETA(DisplayName = "4 x 4"),
-	POT_8 UMETA(DisplayName = "8 x 8"),
-	POT_16 UMETA(DisplayName = "16 x 16"),
-	POT_32 UMETA(DisplayName = "32 x 32"),
-	POT_64 UMETA(DisplayName = "64 x 64"),
-	POT_128 UMETA(DisplayName = "128 x 128"),
-	POT_256 UMETA(DisplayName = "256 x 256"),
-	POT_512 UMETA(DisplayName = "512 x 512"),
-	POT_1024 UMETA(DisplayName = "1024 x 1024"),
-	POT_2048 UMETA(DisplayName = "2048 x 2048"),
-	POT_4096 UMETA(DisplayName = "4096 x 4096"),
-	POT_8192 UMETA(DisplayName = "8192 x 8192")
-};
-
 UENUM(BlueprintType)
 enum class EGLTFMaterialPropertyGroup : uint8
 {
@@ -52,6 +33,24 @@ enum class EGLTFMaterialPropertyGroup : uint8
 	ClearCoatBottomNormal UMETA(DisplayName = "Clear Coat Bottom Normal"),
 };
 
+USTRUCT(Blueprintable)
+struct FGLTFMaterialBakeSize
+{
+	GENERATED_BODY()
+
+	UPROPERTY(EditAnywhere, BlueprintReadWrite, Category = "")
+	int32 X = 1;
+
+	UPROPERTY(EditAnywhere, BlueprintReadWrite, Category = "")
+	int32 Y = 1;
+
+	/** If enabled, bake size is based on the largest texture used in the material input's expression graph. If none found, bake size will fall back to the explicit dimensions. */
+	UPROPERTY(EditAnywhere, BlueprintReadWrite, Category = "")
+	bool bAutoDetect = false;
+
+	static FGLTFMaterialBakeSize Default;
+};
+
 USTRUCT(Blueprintable)
 struct FGLTFOverrideMaterialBakeSettings
 {
@@ -65,7 +64,7 @@ struct FGLTFOverrideMaterialBakeSettings
 
 	/** Overrides default size of the baked out texture. */
 	UPROPERTY(EditAnywhere, BlueprintReadWrite, Category = "", meta = (EditCondition = "bOverrideSize"))
-	EGLTFMaterialBakeSizePOT Size;
+	FGLTFMaterialBakeSize Size;
 
 	/** If enabled, default filtering mode will be overridden by the corresponding property. */
 	UPROPERTY(EditAnywhere, BlueprintReadWrite, Category = "", meta = (InlineEditConditionToggle))
@@ -108,7 +107,7 @@ public:
 
 	static const UMaterialInterface* ResolveProxy(const UMaterialInterface* Material);
 
-	static EGLTFMaterialBakeSizePOT GetBakeSizeForPropertyGroup(const UMaterialInterface* Material, EGLTFMaterialPropertyGroup PropertyGroup, EGLTFMaterialBakeSizePOT DefaultValue);
+	static FGLTFMaterialBakeSize GetBakeSizeForPropertyGroup(const UMaterialInterface* Material, EGLTFMaterialPropertyGroup PropertyGroup, FGLTFMaterialBakeSize DefaultValue);
 	static TextureFilter GetBakeFilterForPropertyGroup(const UMaterialInterface* Material, EGLTFMaterialPropertyGroup PropertyGroup, TextureFilter DefaultValue);
 	static TextureAddress GetBakeTilingForPropertyGroup(const UMaterialInterface* Material, EGLTFMaterialPropertyGroup PropertyGroup, TextureAddress DefaultValue);
 
