[Export - glTF] Thin Translucent materials to glTF (KHR_materials_transmission).
IOR usage fixes for glTF Proxy materials.

#rnx

diff --git a/Engine/Plugins/Enterprise/GLTFExporter/Source/GLTFExporter/Private/Converters/GLTFMaterialProxyFactory.cpp b/Engine/Plugins/Enterprise/GLTFExporter/Source/GLTFExporter/Private/Converters/GLTFMaterialProxyFactory.cpp
--- a/Engine/Plugins/Enterprise/GLTFExporter/Source/GLTFExporter/Private/Converters/GLTFMaterialProxyFactory.cpp
+++ b/Engine/Plugins/Enterprise/GLTFExporter/Source/GLTFExporter/Private/Converters/GLTFMaterialProxyFactory.cpp
@@ -102,9 +102,9 @@ void FGLTFMaterialProxyFactory::SetProxyParameters(UMaterialInstanceConstant* Pr
 	SetProxyParameter(ProxyMaterial, FGLTFProxyMaterialInfo::BaseColor, JsonMaterial.PBRMetallicRoughness.BaseColorTexture);
 	SetProxyParameter(ProxyMaterial, FGLTFProxyMaterialInfo::BaseColorFactor, JsonMaterial.PBRMetallicRoughness.BaseColorFactor);
 
-	if (JsonMaterial.ShadingModel == EGLTFJsonShadingModel::Default || 
-		JsonMaterial.ShadingModel == EGLTFJsonShadingModel::ClearCoat ||
-		JsonMaterial.ShadingModel == EGLTFJsonShadingModel::Sheen)
+	if (JsonMaterial.ShadingModel != EGLTFJsonShadingModel::None && 
+		JsonMaterial.ShadingModel != EGLTFJsonShadingModel::Unlit &&
+		(int)JsonMaterial.ShadingModel < (int)EGLTFJsonShadingModel::NumShadingModels)
 	{
 		SetProxyParameter(ProxyMaterial, FGLTFProxyMaterialInfo::Emissive, JsonMaterial.EmissiveTexture);
 		SetProxyParameter(ProxyMaterial, FGLTFProxyMaterialInfo::EmissiveFactor, JsonMaterial.EmissiveFactor);
@@ -120,6 +120,11 @@ void FGLTFMaterialProxyFactory::SetProxyParameters(UMaterialInstanceConstant* Pr
 		SetProxyParameter(ProxyMaterial, FGLTFProxyMaterialInfo::Occlusion, JsonMaterial.OcclusionTexture);
 		SetProxyParameter(ProxyMaterial, FGLTFProxyMaterialInfo::OcclusionStrength, JsonMaterial.OcclusionTexture.Strength);
 
+		SetProxyParameter(ProxyMaterial, FGLTFProxyMaterialInfo::SpecularFactor, JsonMaterial.Specular.Factor);
+		SetProxyParameter(ProxyMaterial, FGLTFProxyMaterialInfo::SpecularTexture, JsonMaterial.Specular.Texture);
+
+		SetProxyParameter(ProxyMaterial, FGLTFProxyMaterialInfo::IOR, JsonMaterial.IOR.Value);
+
 		if (JsonMaterial.ShadingModel == EGLTFJsonShadingModel::ClearCoat)
 		{
 			SetProxyParameter(ProxyMaterial, FGLTFProxyMaterialInfo::ClearCoat, JsonMaterial.ClearCoat.ClearCoatTexture);
@@ -138,12 +143,11 @@ void FGLTFMaterialProxyFactory::SetProxyParameters(UMaterialInstanceConstant* Pr
 			SetProxyParameter(ProxyMaterial, FGLTFProxyMaterialInfo::SheenRoughnessFactor, JsonMaterial.Sheen.RoughnessFactor);
 			SetProxyParameter(ProxyMaterial, FGLTFProxyMaterialInfo::SheenRoughnessTexture, JsonMaterial.Sheen.RoughnessTexture);
 		}
-	}
-
-	if (JsonMaterial.ShadingModel != EGLTFJsonShadingModel::Unlit) //&& != SpecularGlossiness
-	{
-		SetProxyParameter(ProxyMaterial, FGLTFProxyMaterialInfo::SpecularFactor, JsonMaterial.Specular.Factor);
-		SetProxyParameter(ProxyMaterial, FGLTFProxyMaterialInfo::SpecularTexture, JsonMaterial.Specular.Texture);
+		else if (JsonMaterial.ShadingModel == EGLTFJsonShadingModel::Transmission)
+		{
+			SetProxyParameter(ProxyMaterial, FGLTFProxyMaterialInfo::TransmissionFactor, JsonMaterial.Transmission.Factor);
+			SetProxyParameter(ProxyMaterial, FGLTFProxyMaterialInfo::TransmissionTexture, JsonMaterial.Transmission.Texture);
+		}
 	}
 }
 
@@ -335,6 +339,7 @@ UGLTFExportOptions* FGLTFMaterialProxyFactory::CreateExportOptions(const UGLTFPr
 	ExportOptions->ResetToDefault();
 	ExportOptions->bExportProxyMaterials = false;
 	ExportOptions->BakeMaterialInputs = ProxyOptions->bBakeMaterialInputs ? EGLTFMaterialBakeMode::Simple : EGLTFMaterialBakeMode::Disabled;
+	ExportOptions->bExportThinTranslucentMaterials = ProxyOptions->bUseThinTranslucentShadingModel;
 	ExportOptions->DefaultMaterialBakeSize = ProxyOptions->DefaultMaterialBakeSize;
 	ExportOptions->DefaultMaterialBakeFilter = ProxyOptions->DefaultMaterialBakeFilter;
 	ExportOptions->DefaultMaterialBakeTiling = ProxyOptions->DefaultMaterialBakeTiling;
diff --git a/Engine/Plugins/Enterprise/GLTFExporter/Source/GLTFExporter/Private/Json/GLTFJsonMaterial.cpp b/Engine/Plugins/Enterprise/GLTFExporter/Source/GLTFExporter/Private/Json/GLTFJsonMaterial.cpp
--- a/Engine/Plugins/Enterprise/GLTFExporter/Source/GLTFExporter/Private/Json/GLTFJsonMaterial.cpp
+++ b/Engine/Plugins/Enterprise/GLTFExporter/Source/GLTFExporter/Private/Json/GLTFJsonMaterial.cpp
@@ -164,6 +164,19 @@ void FGLTFJsonSheenExtension::WriteObject(IGLTFJsonWriter& Writer) const
 	}
 }
 
+void FGLTFJsonTransmissionExtension::WriteObject(IGLTFJsonWriter& Writer) const
+{
+	if (!FMath::IsNearlyEqual(Factor, 0.f))
+	{
+		Writer.Write(TEXT("transmissionFactor"), Factor);
+	}
+
+	if (Texture.Index != nullptr)
+	{
+		Writer.Write(TEXT("transmissionTexture"), Texture);
+	}
+}
+
 void FGLTFJsonMaterial::WriteObject(IGLTFJsonWriter& Writer) const
 {
 	if (!Name.IsEmpty())
@@ -212,11 +225,14 @@ void FGLTFJsonMaterial::WriteObject(IGLTFJsonWriter& Writer) const
 	}
 
 	const bool HasEmissiveStrength = !FMath::IsNearlyEqual(EmissiveStrength, 1.0f, Writer.DefaultTolerance);
-	if (ShadingModel == EGLTFJsonShadingModel::Unlit || ShadingModel == EGLTFJsonShadingModel::ClearCoat || 
+	if (ShadingModel == EGLTFJsonShadingModel::Unlit ||
+		ShadingModel == EGLTFJsonShadingModel::ClearCoat ||
+		(ShadingModel == EGLTFJsonShadingModel::Transmission && Transmission.HasValue()) ||
 		HasEmissiveStrength || 
 		Specular.HasValue() || 
 		IOR.HasValue() ||
-		Sheen.HasValue())
+		Sheen.HasValue() ||
+		Transmission.HasValue())
 	{
 		Writer.StartExtensions();
 
@@ -245,6 +261,10 @@ void FGLTFJsonMaterial::WriteObject(IGLTFJsonWriter& Writer) const
 		{
 			Writer.Write(EGLTFJsonExtension::KHR_MaterialsClearCoat, ClearCoat);
 		}
+		else if (ShadingModel == EGLTFJsonShadingModel::Transmission)
+		{
+			Writer.Write(EGLTFJsonExtension::KHR_MaterialsTransmission, Transmission);
+		}
 
 		if (HasEmissiveStrength)
 		{
diff --git a/Engine/Plugins/Enterprise/GLTFExporter/Source/GLTFExporter/Private/Materials/GLTFProxyMaterialInfo.cpp b/Engine/Plugins/Enterprise/GLTFExporter/Source/GLTFExporter/Private/Materials/GLTFProxyMaterialInfo.cpp
--- a/Engine/Plugins/Enterprise/GLTFExporter/Source/GLTFExporter/Private/Materials/GLTFProxyMaterialInfo.cpp
+++ b/Engine/Plugins/Enterprise/GLTFExporter/Source/GLTFExporter/Private/Materials/GLTFProxyMaterialInfo.cpp
@@ -31,7 +31,12 @@ const TGLTFProxyMaterialParameterInfo<float> FGLTFProxyMaterialInfo::ClearCoatNo
 const TGLTFProxyMaterialParameterInfo<float> FGLTFProxyMaterialInfo::SpecularFactor = { TEXT("Specular Factor") };
 const FGLTFProxyMaterialTextureParameterInfo FGLTFProxyMaterialInfo::SpecularTexture = { TEXT("Specular") };
 
+const TGLTFProxyMaterialParameterInfo<float> FGLTFProxyMaterialInfo::IOR = { TEXT("IOR") };
+
 const TGLTFProxyMaterialParameterInfo<FLinearColor> FGLTFProxyMaterialInfo::SheenColorFactor = { TEXT("Sheen Color Factor") };
 const FGLTFProxyMaterialTextureParameterInfo FGLTFProxyMaterialInfo::SheenColorTexture = { TEXT("Sheen Color") };
 const TGLTFProxyMaterialParameterInfo<float> FGLTFProxyMaterialInfo::SheenRoughnessFactor = { TEXT("Sheen Roughness Factor") };
 const FGLTFProxyMaterialTextureParameterInfo FGLTFProxyMaterialInfo::SheenRoughnessTexture = { TEXT("Sheen Roughness") };
+
+const TGLTFProxyMaterialParameterInfo<float> FGLTFProxyMaterialInfo::TransmissionFactor = { TEXT("Transmission Factor") };
+const FGLTFProxyMaterialTextureParameterInfo FGLTFProxyMaterialInfo::TransmissionTexture = { TEXT("Transmission") };
diff --git a/Engine/Plugins/Enterprise/GLTFExporter/Source/GLTFExporter/Private/Materials/GLTFProxyMaterialInfo.h b/Engine/Plugins/Enterprise/GLTFExporter/Source/GLTFExporter/Private/Materials/GLTFProxyMaterialInfo.h
--- a/Engine/Plugins/Enterprise/GLTFExporter/Source/GLTFExporter/Private/Materials/GLTFProxyMaterialInfo.h
+++ b/Engine/Plugins/Enterprise/GLTFExporter/Source/GLTFExporter/Private/Materials/GLTFProxyMaterialInfo.h
@@ -37,8 +37,13 @@ public:
 	static const TGLTFProxyMaterialParameterInfo<float> SpecularFactor;
 	static const FGLTFProxyMaterialTextureParameterInfo SpecularTexture; //Only using Alpha Channel
 
+	static const TGLTFProxyMaterialParameterInfo<float> IOR;
+
 	static const TGLTFProxyMaterialParameterInfo<FLinearColor> SheenColorFactor;
 	static const FGLTFProxyMaterialTextureParameterInfo        SheenColorTexture; //RGB
 	static const TGLTFProxyMaterialParameterInfo<float>        SheenRoughnessFactor;
 	static const FGLTFProxyMaterialTextureParameterInfo        SheenRoughnessTexture; //A
+
+	static const TGLTFProxyMaterialParameterInfo<float> TransmissionFactor;
+	static const FGLTFProxyMaterialTextureParameterInfo TransmissionTexture; //Only using Red Channel
 };
diff --git a/Engine/Plugins/Enterprise/GLTFExporter/Source/GLTFExporter/Private/Options/GLTFExportOptions.cpp b/Engine/Plugins/Enterprise/GLTFExporter/Source/GLTFExporter/Private/Options/GLTFExportOptions.cpp
--- a/Engine/Plugins/Enterprise/GLTFExporter/Source/GLTFExporter/Private/Options/GLTFExportOptions.cpp
+++ b/Engine/Plugins/Enterprise/GLTFExporter/Source/GLTFExporter/Private/Options/GLTFExportOptions.cpp
@@ -18,6 +18,7 @@ void UGLTFExportOptions::ResetToDefault()
 	bExportUnlitMaterials = true;
 	bExportClearCoatMaterials = true;
 	bExportClothMaterials = true;
+	bExportThinTranslucentMaterials = false;
 	bExportEmissiveStrength = true;
 	BakeMaterialInputs = EGLTFMaterialBakeMode::UseMeshData;
 	DefaultMaterialBakeSize = FGLTFMaterialBakeSize::Default;
diff --git a/Engine/Plugins/Enterprise/GLTFExporter/Source/GLTFExporter/Private/Options/GLTFProxyOptions.cpp b/Engine/Plugins/Enterprise/GLTFExporter/Source/GLTFExporter/Private/Options/GLTFProxyOptions.cpp
--- a/Engine/Plugins/Enterprise/GLTFExporter/Source/GLTFExporter/Private/Options/GLTFProxyOptions.cpp
+++ b/Engine/Plugins/Enterprise/GLTFExporter/Source/GLTFExporter/Private/Options/GLTFProxyOptions.cpp
@@ -13,6 +13,7 @@ UGLTFProxyOptions::UGLTFProxyOptions(const FObjectInitializer& ObjectInitializer
 void UGLTFProxyOptions::ResetToDefault()
 {
 	bBakeMaterialInputs = true;
+	bUseThinTranslucentShadingModel = false;
 	DefaultMaterialBakeSize = FGLTFMaterialBakeSize::Default;
 	DefaultMaterialBakeFilter = TF_Trilinear;
 	DefaultMaterialBakeTiling = TA_Wrap;
diff --git a/Engine/Plugins/Enterprise/GLTFExporter/Source/GLTFExporter/Private/Tasks/GLTFDelayedMaterialTasks.cpp b/Engine/Plugins/Enterprise/GLTFExporter/Source/GLTFExporter/Private/Tasks/GLTFDelayedMaterialTasks.cpp
--- a/Engine/Plugins/Enterprise/GLTFExporter/Source/GLTFExporter/Private/Tasks/GLTFDelayedMaterialTasks.cpp
+++ b/Engine/Plugins/Enterprise/GLTFExporter/Source/GLTFExporter/Private/Tasks/GLTFDelayedMaterialTasks.cpp
@@ -63,6 +63,12 @@ void FGLTFDelayedMaterialTask::Process()
 	ConvertShadingModel(JsonMaterial->ShadingModel);
 	ConvertAlphaMode(JsonMaterial->AlphaMode);
 
+	if (!ensure((JsonMaterial->ShadingModel != EGLTFJsonShadingModel::None) && ((int)JsonMaterial->ShadingModel < (int)EGLTFJsonShadingModel::NumShadingModels)))
+	{
+		Builder.LogWarning(FString::Printf(TEXT("Failed to export %s material. ShadingModel is unexpected."), *Material->GetName()));
+		return;
+	}
+
 	if (FGLTFProxyMaterialUtilities::IsProxyMaterial(BaseMaterial))
 	{
 		GetProxyParameters(*JsonMaterial);
@@ -70,39 +76,45 @@ void FGLTFDelayedMaterialTask::Process()
 	}
 
 #if WITH_EDITOR
-	if (JsonMaterial->ShadingModel != EGLTFJsonShadingModel::None)
 	{
-		const FMaterialPropertyEx BaseColorProperty = JsonMaterial->ShadingModel == EGLTFJsonShadingModel::Unlit ? MP_EmissiveColor : MP_BaseColor;
-		const FMaterialPropertyEx OpacityProperty = JsonMaterial->AlphaMode == EGLTFJsonAlphaMode::Mask ? MP_OpacityMask : MP_Opacity;
-
-		// TODO: check if a property is active before trying to get it (i.e. Material->IsPropertyActive)
+		if (JsonMaterial->ShadingModel == EGLTFJsonShadingModel::Transmission)
+		{
+			const FMaterialPropertyEx BaseColorProperty = MP_BaseColor;
+			const FMaterialPropertyEx OpacityProperty = MP_Opacity;
 
-		if (JsonMaterial->AlphaMode == EGLTFJsonAlphaMode::Opaque)
+			TryGetTransmissionBaseColorAndOpacity(*JsonMaterial, BaseColorProperty, OpacityProperty);
+		}
+		else
 		{
-			if (!TryGetConstantColor(JsonMaterial->PBRMetallicRoughness.BaseColorFactor, BaseColorProperty))
+			const FMaterialPropertyEx BaseColorProperty = JsonMaterial->ShadingModel == EGLTFJsonShadingModel::Unlit ? MP_EmissiveColor : MP_BaseColor;
+			const FMaterialPropertyEx OpacityProperty = JsonMaterial->AlphaMode == EGLTFJsonAlphaMode::Mask ? MP_OpacityMask : MP_Opacity;
+
+			// TODO: check if a property is active before trying to get it (i.e. Material->IsPropertyActive)
+			if (JsonMaterial->AlphaMode == EGLTFJsonAlphaMode::Opaque)
 			{
-				if (!TryGetSourceTexture(JsonMaterial->PBRMetallicRoughness.BaseColorTexture, BaseColorProperty, DefaultColorInputMasks))
+				if (!TryGetConstantColor(JsonMaterial->PBRMetallicRoughness.BaseColorFactor, BaseColorProperty))
 				{
-					if (!TryGetBakedMaterialProperty(JsonMaterial->PBRMetallicRoughness.BaseColorTexture, JsonMaterial->PBRMetallicRoughness.BaseColorFactor, BaseColorProperty, TEXT("BaseColor")))
+					if (!TryGetSourceTexture(JsonMaterial->PBRMetallicRoughness.BaseColorTexture, BaseColorProperty, DefaultColorInputMasks))
 					{
-						Builder.LogWarning(FString::Printf(TEXT("Failed to export %s for material %s"), *BaseColorProperty.ToString(), *Material->GetName()));
+						if (!TryGetBakedMaterialProperty(JsonMaterial->PBRMetallicRoughness.BaseColorTexture, JsonMaterial->PBRMetallicRoughness.BaseColorFactor, BaseColorProperty, TEXT("BaseColor")))
+						{
+							Builder.LogWarning(FString::Printf(TEXT("Failed to export %s for material %s"), *BaseColorProperty.ToString(), *Material->GetName()));
+						}
 					}
 				}
-			}
 
-			JsonMaterial->PBRMetallicRoughness.BaseColorFactor.A = 1.0f; // make sure base color is opaque
-		}
-		else
-		{
-			if (!TryGetBaseColorAndOpacity(JsonMaterial->PBRMetallicRoughness, BaseColorProperty, OpacityProperty))
+				JsonMaterial->PBRMetallicRoughness.BaseColorFactor.A = 1.0f; // make sure base color is opaque
+			}
+			else
 			{
-				Builder.LogWarning(FString::Printf(TEXT("Failed to export %s and %s for material %s"), *BaseColorProperty.ToString(), *OpacityProperty.ToString(), *Material->GetName()));
+				if (!TryGetBaseColorAndOpacity(JsonMaterial->PBRMetallicRoughness, BaseColorProperty, OpacityProperty))
+				{
+					Builder.LogWarning(FString::Printf(TEXT("Failed to export %s and %s for material %s"), *BaseColorProperty.ToString(), *OpacityProperty.ToString(), *Material->GetName()));
+				}
 			}
 		}
 
-		if (JsonMaterial->ShadingModel == EGLTFJsonShadingModel::Default || 
-			JsonMaterial->ShadingModel == EGLTFJsonShadingModel::ClearCoat ||
-			JsonMaterial->ShadingModel == EGLTFJsonShadingModel::Sheen)
+		if (JsonMaterial->ShadingModel != EGLTFJsonShadingModel::Unlit)
 		{
 			const FMaterialPropertyEx MetallicProperty = MP_Metallic;
 			const FMaterialPropertyEx RoughnessProperty = MP_Roughness;
@@ -153,6 +165,7 @@ void FGLTFDelayedMaterialTask::Process()
 
 			if (JsonMaterial->AlphaMode == EGLTFJsonAlphaMode::Blend)
 			{
+				bool bRefractionSet = false;
 				if (BaseMaterial->RefractionMethod == ERefractionMode::RM_IndexOfRefraction)
 				{
 					const FMaterialPropertyEx RefractionProperty = MP_Refraction;
@@ -162,8 +175,17 @@ void FGLTFDelayedMaterialTask::Process()
 						{
 							Builder.LogWarning(FString::Printf(TEXT("Failed to export %s for material %s"), *RefractionProperty.ToString(), *Material->GetName()));
 						}
+						else
+						{
+							bRefractionSet = true;
+						}
 					}
 				}
+				if (!bRefractionSet)
+				{
+					// in this case UE takes Refraction as IOR 1.0 (Air)
+					JsonMaterial->IOR.Value = 1.0f;
+				}
 			}
 
 			if (JsonMaterial->ShadingModel == EGLTFJsonShadingModel::ClearCoat)
@@ -254,9 +276,7 @@ void FGLTFDelayedMaterialTask::GetProxyParameters(FGLTFJsonMaterial& OutMaterial
 	GetProxyParameter(FGLTFProxyMaterialInfo::BaseColor, OutMaterial.PBRMetallicRoughness.BaseColorTexture);
 	GetProxyParameter(FGLTFProxyMaterialInfo::BaseColorFactor, OutMaterial.PBRMetallicRoughness.BaseColorFactor);
 
-	if (OutMaterial.ShadingModel == EGLTFJsonShadingModel::Default || 
-		OutMaterial.ShadingModel == EGLTFJsonShadingModel::ClearCoat ||
-		OutMaterial.ShadingModel == EGLTFJsonShadingModel::Sheen)
+	if (OutMaterial.ShadingModel != EGLTFJsonShadingModel::Unlit)
 	{
 		GetProxyParameter(FGLTFProxyMaterialInfo::Emissive, OutMaterial.EmissiveTexture);
 		GetProxyParameter(FGLTFProxyMaterialInfo::EmissiveFactor, OutMaterial.EmissiveFactor);
@@ -284,6 +304,11 @@ void FGLTFDelayedMaterialTask::GetProxyParameters(FGLTFJsonMaterial& OutMaterial
 		GetProxyParameter(FGLTFProxyMaterialInfo::Occlusion, OutMaterial.OcclusionTexture);
 		GetProxyParameter(FGLTFProxyMaterialInfo::OcclusionStrength, OutMaterial.OcclusionTexture.Strength);
 
+		GetProxyParameter(FGLTFProxyMaterialInfo::SpecularFactor, OutMaterial.Specular.Factor);
+		GetProxyParameter(FGLTFProxyMaterialInfo::SpecularTexture, OutMaterial.Specular.Texture);
+
+		GetProxyParameter(FGLTFProxyMaterialInfo::IOR, OutMaterial.IOR.Value);
+
 		if (OutMaterial.ShadingModel == EGLTFJsonShadingModel::ClearCoat)
 		{
 			GetProxyParameter(FGLTFProxyMaterialInfo::ClearCoat, OutMaterial.ClearCoat.ClearCoatTexture);
@@ -302,12 +327,11 @@ void FGLTFDelayedMaterialTask::GetProxyParameters(FGLTFJsonMaterial& OutMaterial
 			GetProxyParameter(FGLTFProxyMaterialInfo::SheenRoughnessFactor, OutMaterial.Sheen.RoughnessFactor);
 			GetProxyParameter(FGLTFProxyMaterialInfo::SheenRoughnessTexture, OutMaterial.Sheen.RoughnessTexture);
 		}
-	}
-
-	if (OutMaterial.ShadingModel != EGLTFJsonShadingModel::Unlit)
-	{
-		GetProxyParameter(FGLTFProxyMaterialInfo::SpecularFactor, OutMaterial.Specular.Factor);
-		GetProxyParameter(FGLTFProxyMaterialInfo::SpecularTexture, OutMaterial.Specular.Texture);
+		else if (OutMaterial.ShadingModel == EGLTFJsonShadingModel::Transmission)
+		{
+			GetProxyParameter(FGLTFProxyMaterialInfo::TransmissionFactor, OutMaterial.Transmission.Factor);
+			GetProxyParameter(FGLTFProxyMaterialInfo::TransmissionTexture, OutMaterial.Transmission.Texture);
+		}
 	}
 }
 
@@ -510,6 +534,17 @@ void FGLTFDelayedMaterialTask::ConvertShadingModel(EGLTFJsonShadingModel& OutSha
 			*Material->GetName(),
 			*FGLTFNameUtilities::GetName(MSM_DefaultLit)));
 	}
+
+	if (OutShadingModel == EGLTFJsonShadingModel::Transmission && !Builder.ExportOptions->bExportThinTranslucentMaterials)
+	{
+		OutShadingModel = EGLTFJsonShadingModel::Default;
+
+		Builder.LogWarning(FString::Printf(
+			TEXT("Shading model (%s) in material %s disabled by export options, will export as %s"),
+			*FGLTFNameUtilities::GetName(ShadingModel),
+			*Material->GetName(),
+			*FGLTFNameUtilities::GetName(MSM_DefaultLit)));
+	}
 }
 
 void FGLTFDelayedMaterialTask::ConvertAlphaMode(EGLTFJsonAlphaMode& OutAlphaMode) const
@@ -987,7 +1022,16 @@ bool FGLTFDelayedMaterialTask::TryGetSpecular(FGLTFJsonMaterial& OutMaterial, co
 	}
 	else
 	{
-		if (!TryGetBakedMaterialPropertyOntoAlphaChannel(OutMaterial.Specular.Texture, OutMaterial.Specular.Factor, 1.0f, SpecularProperty))
+		if (!TryGetBakedMaterialProperty(OutMaterial.Specular.Texture, OutMaterial.Specular.Factor, 1.0f, SpecularProperty, 
+			[](FColor& Color, const uint8& ChannelValueToMove)
+			{
+				Color.A = ChannelValueToMove;
+
+				Color.R = 255;
+				Color.G = 255;
+				Color.B = 255;
+			}
+		))
 		{
 			return false;
 		}
@@ -1069,7 +1113,16 @@ bool FGLTFDelayedMaterialTask::TryGetFuzzColorAndCloth(FGLTFJsonMaterial& OutMat
 	{
 		if (!TryGetSourceTexture(OutMaterial.Sheen.RoughnessTexture, ClothProperty, { AlphaMask }))
 		{
-			if (!TryGetBakedMaterialPropertyOntoAlphaChannel(OutMaterial.Sheen.RoughnessTexture, OutMaterial.Sheen.RoughnessFactor, 1.0f, ClothProperty, TEXT("SheenRoughnessTexture")))
+			if (!TryGetBakedMaterialProperty(OutMaterial.Sheen.RoughnessTexture, OutMaterial.Sheen.RoughnessFactor, 1.0f, ClothProperty,
+				[](FColor& Color, const uint8& ChannelValueToMove)
+				{
+					Color.A = ChannelValueToMove;
+
+					Color.R = 255;
+					Color.G = 255;
+					Color.B = 255;
+				},
+				TEXT("SheenRoughnessTexture")))
 			{
 				return false;
 			}
@@ -1083,6 +1136,49 @@ bool FGLTFDelayedMaterialTask::TryGetFuzzColorAndCloth(FGLTFJsonMaterial& OutMat
 	return true;
 }
 
+bool FGLTFDelayedMaterialTask::TryGetTransmissionBaseColorAndOpacity(FGLTFJsonMaterial& OutMaterial, const FMaterialPropertyEx& BaseColorProperty, const FMaterialPropertyEx& OpacityProperty)
+{
+	const bool bIsBaseColorConstant = TryGetConstantColor(OutMaterial.PBRMetallicRoughness.BaseColorFactor, BaseColorProperty);
+	
+	float Opacity = 0.0f;
+	const bool bIsOpacityConstant = TryGetConstantScalar(Opacity, OpacityProperty);
+
+	if (bIsBaseColorConstant && bIsOpacityConstant)
+	{
+		OutMaterial.Transmission.Factor = 1 - Opacity;
+		OutMaterial.PBRMetallicRoughness.BaseColorFactor.A = 1.0f;
+
+		return true;
+	}
+
+	if (!bIsBaseColorConstant)
+	{
+		if (!TryGetBakedMaterialProperty(OutMaterial.PBRMetallicRoughness.BaseColorTexture, OutMaterial.PBRMetallicRoughness.BaseColorFactor, BaseColorProperty, TEXT("BaseColor")))
+		{
+			return false;
+		}
+	}
+
+	if (!bIsOpacityConstant)
+	{
+		if (!TryGetBakedMaterialProperty(OutMaterial.Transmission.Texture, OutMaterial.Transmission.Factor, 1.0f, OpacityProperty,
+			[](FColor& Color, const uint8& ChannelValueToMove)
+			{
+				Color.R = ChannelValueToMove;
+
+				Color.G = 255;
+				Color.B = 255;
+				Color.A = 255;
+			},
+			TEXT("TransmissionTexture")))
+		{
+			return false;
+		}
+	}
+
+	return true;
+}
+
 bool FGLTFDelayedMaterialTask::IsPropertyNonDefault(const FMaterialPropertyEx& Property) const
 {
 	// Custom outputs are by definition standalone and never part of material attributes
@@ -1653,7 +1749,8 @@ bool FGLTFDelayedMaterialTask::TryGetBakedMaterialProperty(FGLTFJsonTextureInfo&
 	return true;
 }
 
-bool FGLTFDelayedMaterialTask::TryGetBakedMaterialPropertyOntoAlphaChannel(FGLTFJsonTextureInfo& Texture, float& Factor, const float& DefaultFactorValue, const FMaterialPropertyEx& Property, const FString& TextureName)
+template <typename CallbackType>
+bool FGLTFDelayedMaterialTask::TryGetBakedMaterialProperty(FGLTFJsonTextureInfo& Texture, float& Factor, const float& DefaultFactorValue, const FMaterialPropertyEx& Property, CallbackType Callback /*ModifyPixel(Pixel&,ChannelValue)*/, const FString& TextureName)
 {
 	if (Builder.ExportOptions->BakeMaterialInputs == EGLTFMaterialBakeMode::Disabled)
 	{
@@ -1715,37 +1812,31 @@ bool FGLTFDelayedMaterialTask::TryGetBakedMaterialPropertyOntoAlphaChannel(FGLTF
 			case RGB:
 				for (FColor& Pixel : *PropertyBakeOutput.Pixels)
 				{
-					Pixel.A = Pixel.R;
-
-					Pixel.R = 255;
-					Pixel.G = 255;
-					Pixel.B = 255;
+					Callback(Pixel, Pixel.R);
 				}
 				break;
 
 			case GREEN:
 				for (FColor& Pixel : *PropertyBakeOutput.Pixels)
 				{
-					Pixel.A = Pixel.G;
-
-					Pixel.R = 255;
-					Pixel.G = 255;
-					Pixel.B = 255;
+					Callback(Pixel, Pixel.G);
 				}
 				break;
 
 			case BLUE:
 				for (FColor& Pixel : *PropertyBakeOutput.Pixels)
 				{
-					Pixel.A = Pixel.B;
-
-					Pixel.R = 255;
-					Pixel.G = 255;
-					Pixel.B = 255;
+					Callback(Pixel, Pixel.B);
 				}
 				break;
 
 			case ALPHA:
+				for (FColor& Pixel : *PropertyBakeOutput.Pixels)
+				{
+					Callback(Pixel, Pixel.A);
+				}
+				break;
+
 			default:
 				break;
 			}
diff --git a/Engine/Plugins/Enterprise/GLTFExporter/Source/GLTFExporter/Private/Tasks/GLTFDelayedMaterialTasks.h b/Engine/Plugins/Enterprise/GLTFExporter/Source/GLTFExporter/Private/Tasks/GLTFDelayedMaterialTasks.h
--- a/Engine/Plugins/Enterprise/GLTFExporter/Source/GLTFExporter/Private/Tasks/GLTFDelayedMaterialTasks.h
+++ b/Engine/Plugins/Enterprise/GLTFExporter/Source/GLTFExporter/Private/Tasks/GLTFDelayedMaterialTasks.h
@@ -78,6 +78,7 @@ private:
 	bool TryGetSpecular(FGLTFJsonMaterial& OutMaterial, const FMaterialPropertyEx& SpecularProperty);
 	bool TryGetRefraction(FGLTFJsonMaterial& OutMaterial, const FMaterialPropertyEx& RefractionProperty);
 	bool TryGetFuzzColorAndCloth(FGLTFJsonMaterial& OutMaterial, const FMaterialPropertyEx& FuzzColorProperty, const FMaterialPropertyEx& ClothProperty);
+	bool TryGetTransmissionBaseColorAndOpacity(FGLTFJsonMaterial& OutMaterial, const FMaterialPropertyEx& BaseColorProperty, const FMaterialPropertyEx& OpacityProperty);
 
 	bool IsPropertyNonDefault(const FMaterialPropertyEx& Property) const;
 	bool TryGetConstantColor(FGLTFJsonColor3& OutValue, const FMaterialPropertyEx& Property) const;
@@ -92,7 +93,9 @@ private:
 	bool TryGetBakedMaterialProperty(FGLTFJsonTextureInfo& OutTexInfo, FGLTFJsonColor4& OutConstant, const FMaterialPropertyEx& Property, const FString& PropertyName);
 	bool TryGetBakedMaterialProperty(FGLTFJsonTextureInfo& OutTexInfo, float& OutConstant, const FMaterialPropertyEx& Property, const FString& PropertyName);
 	bool TryGetBakedMaterialProperty(FGLTFJsonTextureInfo& OutTexInfo, const FMaterialPropertyEx& Property, const FString& PropertyName);
-	bool TryGetBakedMaterialPropertyOntoAlphaChannel(FGLTFJsonTextureInfo& Texture, float& Factor, const float& DefaultFactorValue, const FMaterialPropertyEx& Property, const FString& TextureName = TEXT(""));
+
+	template <typename CallbackType>
+	bool TryGetBakedMaterialProperty(FGLTFJsonTextureInfo& Texture, float& Factor, const float& DefaultFactorValue, const FMaterialPropertyEx& Property, CallbackType Callback /*ModifyPixel(Pixel&,ChannelValue)*/, const FString& TextureName = TEXT(""));
 
 	FGLTFPropertyBakeOutput BakeMaterialProperty(const FMaterialPropertyEx& Property, int32& OutTexCoord, FGLTFJsonTextureTransform& OutTransform);
 	FGLTFPropertyBakeOutput BakeMaterialProperty(const FMaterialPropertyEx& Property, int32& OutTexCoord, FGLTFJsonTextureTransform& OutTransform, const FIntPoint& TextureSize, bool bFillAlpha);
diff --git a/Engine/Plugins/Enterprise/GLTFExporter/Source/GLTFExporter/Private/Utilities/GLTFCoreUtilities.cpp b/Engine/Plugins/Enterprise/GLTFExporter/Source/GLTFExporter/Private/Utilities/GLTFCoreUtilities.cpp
--- a/Engine/Plugins/Enterprise/GLTFExporter/Source/GLTFExporter/Private/Utilities/GLTFCoreUtilities.cpp
+++ b/Engine/Plugins/Enterprise/GLTFExporter/Source/GLTFExporter/Private/Utilities/GLTFCoreUtilities.cpp
@@ -219,11 +219,12 @@ EGLTFJsonShadingModel FGLTFCoreUtilities::ConvertShadingModel(EMaterialShadingMo
 {
 	switch (ShadingModel)
 	{
-		case MSM_Unlit:      return EGLTFJsonShadingModel::Unlit;
-		case MSM_DefaultLit: return EGLTFJsonShadingModel::Default;
-		case MSM_ClearCoat:  return EGLTFJsonShadingModel::ClearCoat;
-		case MSM_Cloth:      return EGLTFJsonShadingModel::Sheen;
-		default:             return EGLTFJsonShadingModel::None;
+		case MSM_Unlit:                return EGLTFJsonShadingModel::Unlit;
+		case MSM_DefaultLit:           return EGLTFJsonShadingModel::Default;
+		case MSM_ClearCoat:            return EGLTFJsonShadingModel::ClearCoat;
+		case MSM_Cloth:                return EGLTFJsonShadingModel::Sheen;
+		case MSM_ThinTranslucent:      return EGLTFJsonShadingModel::Transmission;
+		default:                       return EGLTFJsonShadingModel::None;
 	}
 }
 
diff --git a/Engine/Plugins/Enterprise/GLTFExporter/Source/GLTFExporter/Private/Utilities/GLTFJsonUtilities.cpp b/Engine/Plugins/Enterprise/GLTFExporter/Source/GLTFExporter/Private/Utilities/GLTFJsonUtilities.cpp
--- a/Engine/Plugins/Enterprise/GLTFExporter/Source/GLTFExporter/Private/Utilities/GLTFJsonUtilities.cpp
+++ b/Engine/Plugins/Enterprise/GLTFExporter/Source/GLTFExporter/Private/Utilities/GLTFJsonUtilities.cpp
@@ -13,6 +13,7 @@ const TCHAR* FGLTFJsonUtilities::GetValue(EGLTFJsonExtension Enum)
 		case EGLTFJsonExtension::KHR_MaterialsVariants:         return TEXT("KHR_materials_variants");
 		case EGLTFJsonExtension::KHR_MaterialsIOR:              return TEXT("KHR_materials_ior");
 		case EGLTFJsonExtension::KHR_MaterialsSheen:            return TEXT("KHR_materials_sheen");
+		case EGLTFJsonExtension::KHR_MaterialsTransmission:     return TEXT("KHR_materials_transmission");
 		case EGLTFJsonExtension::KHR_MeshQuantization:          return TEXT("KHR_mesh_quantization");
 		case EGLTFJsonExtension::KHR_TextureTransform:          return TEXT("KHR_texture_transform");
 		case EGLTFJsonExtension::KHR_MaterialsSpecular:	        return TEXT("KHR_materials_specular");
@@ -120,10 +121,11 @@ const TCHAR* FGLTFJsonUtilities::GetValue(EGLTFJsonShadingModel Enum)
 {
 	switch (Enum)
 	{
-		case EGLTFJsonShadingModel::Default:   return TEXT("Default");
-		case EGLTFJsonShadingModel::Unlit:     return TEXT("Unlit");
-		case EGLTFJsonShadingModel::ClearCoat: return TEXT("ClearCoat");
-		case EGLTFJsonShadingModel::Sheen:     return TEXT("Sheen");
+		case EGLTFJsonShadingModel::Default:          return TEXT("Default");
+		case EGLTFJsonShadingModel::Unlit:            return TEXT("Unlit");
+		case EGLTFJsonShadingModel::ClearCoat:        return TEXT("ClearCoat");
+		case EGLTFJsonShadingModel::Sheen:            return TEXT("Sheen");
+		case EGLTFJsonShadingModel::Transmission:     return TEXT("Transmission");
 		default:
 			checkNoEntry();
 			return TEXT("");
diff --git a/Engine/Plugins/Enterprise/GLTFExporter/Source/GLTFExporter/Public/Json/GLTFJsonEnums.h b/Engine/Plugins/Enterprise/GLTFExporter/Source/GLTFExporter/Public/Json/GLTFJsonEnums.h
--- a/Engine/Plugins/Enterprise/GLTFExporter/Source/GLTFExporter/Public/Json/GLTFJsonEnums.h
+++ b/Engine/Plugins/Enterprise/GLTFExporter/Source/GLTFExporter/Public/Json/GLTFJsonEnums.h
@@ -11,6 +11,7 @@ enum class EGLTFJsonExtension
 	KHR_MaterialsVariants,
 	KHR_MaterialsIOR,
 	KHR_MaterialsSheen,
+	KHR_MaterialsTransmission,
 	KHR_MeshQuantization,
 	KHR_TextureTransform,
 	KHR_MaterialsSpecular
@@ -23,6 +24,7 @@ enum class EGLTFJsonShadingModel
 	Unlit,
 	ClearCoat,
 	Sheen, /*UE Cloth*/
+	Transmission, /*UE Thin Translucent*/
 
 	NumShadingModels
 };
diff --git a/Engine/Plugins/Enterprise/GLTFExporter/Source/GLTFExporter/Public/Json/GLTFJsonMaterial.h b/Engine/Plugins/Enterprise/GLTFExporter/Source/GLTFExporter/Public/Json/GLTFJsonMaterial.h
--- a/Engine/Plugins/Enterprise/GLTFExporter/Source/GLTFExporter/Public/Json/GLTFJsonMaterial.h
+++ b/Engine/Plugins/Enterprise/GLTFExporter/Source/GLTFExporter/Public/Json/GLTFJsonMaterial.h
@@ -135,6 +135,21 @@ struct GLTFEXPORTER_API FGLTFJsonSheenExtension : IGLTFJsonObject
 	virtual void WriteObject(IGLTFJsonWriter& Writer) const override;
 };
 
+struct GLTFEXPORTER_API FGLTFJsonTransmissionExtension : IGLTFJsonObject
+{
+	float                   Factor; //transmissionFactor
+	FGLTFJsonTextureInfo    Texture; //transmissionTexture
+
+	FGLTFJsonTransmissionExtension()
+		: Factor(0.f)
+	{
+	}
+
+	bool HasValue() const { return !FMath::IsNearlyEqual(Factor, 0.f) || Texture.Index != nullptr; }
+
+	virtual void WriteObject(IGLTFJsonWriter& Writer) const override;
+};
+
 struct GLTFEXPORTER_API FGLTFJsonMaterial : IGLTFJsonIndexedObject
 {
 	FString Name;
@@ -155,10 +170,11 @@ struct GLTFEXPORTER_API FGLTFJsonMaterial : IGLTFJsonIndexedObject
 
 	bool DoubleSided;
 
-	FGLTFJsonClearCoatExtension ClearCoat;
-	FGLTFJsonSpecularExtension  Specular;
-	FGLTFJsonIORExtension       IOR;
-	FGLTFJsonSheenExtension     Sheen;
+	FGLTFJsonClearCoatExtension    ClearCoat;
+	FGLTFJsonSpecularExtension     Specular;
+	FGLTFJsonIORExtension          IOR;
+	FGLTFJsonSheenExtension        Sheen;
+	FGLTFJsonTransmissionExtension Transmission;
 
 	virtual void WriteObject(IGLTFJsonWriter& Writer) const override;
 
diff --git a/Engine/Plugins/Enterprise/GLTFExporter/Source/GLTFExporter/Public/Options/GLTFExportOptions.h b/Engine/Plugins/Enterprise/GLTFExporter/Source/GLTFExporter/Public/Options/GLTFExportOptions.h
--- a/Engine/Plugins/Enterprise/GLTFExporter/Source/GLTFExporter/Public/Options/GLTFExportOptions.h
+++ b/Engine/Plugins/Enterprise/GLTFExporter/Source/GLTFExporter/Public/Options/GLTFExportOptions.h
@@ -65,6 +65,10 @@ class GLTFEXPORTER_API UGLTFExportOptions : public UObject
 	UPROPERTY(EditAnywhere, BlueprintReadWrite, Config, Category = Material)
 	bool bExportClothMaterials;
 
+	/** If enabled, materials with shading model thin translucency will be exported. Export is only partial. Uses extension KHR_materials_transmission. */
+	UPROPERTY(EditAnywhere, BlueprintReadWrite, Config, Category = Material)
+	bool bExportThinTranslucentMaterials;
+
 	/** If enabled, allows materials to have an emissive factor that exceeds the standard range [0.0, 1.0]. Uses extension KHR_materials_emissive_strength. */
 	UPROPERTY(EditAnywhere, BlueprintReadWrite, Config, Category = Material)
 	bool bExportEmissiveStrength;
diff --git a/Engine/Plugins/Enterprise/GLTFExporter/Source/GLTFExporter/Public/Options/GLTFProxyOptions.h b/Engine/Plugins/Enterprise/GLTFExporter/Source/GLTFExporter/Public/Options/GLTFProxyOptions.h
--- a/Engine/Plugins/Enterprise/GLTFExporter/Source/GLTFExporter/Public/Options/GLTFProxyOptions.h
+++ b/Engine/Plugins/Enterprise/GLTFExporter/Source/GLTFExporter/Public/Options/GLTFProxyOptions.h
@@ -14,6 +14,10 @@ class GLTFEXPORTER_API UGLTFProxyOptions : public UObject
 	UPROPERTY(EditAnywhere, BlueprintReadWrite, Config, Category = Material)
 	bool bBakeMaterialInputs;
 
+	/** If enabled, materials with shading model thin translucency will be used. Conversion is only partial.*/
+	UPROPERTY(EditAnywhere, BlueprintReadWrite, Config, Category = Material)
+	bool bUseThinTranslucentShadingModel;
+
 	/** Default size of the baked out texture (containing the material input). Can be overridden by material- and input-specific bake settings, see GLTFMaterialExportOptions. */
 	UPROPERTY(EditAnywhere, BlueprintReadWrite, Config, Category = Material, Meta = (EditCondition = "bBakeMaterialInputs"))
 	FGLTFMaterialBakeSize DefaultMaterialBakeSize;
