[Export - glTF] Bug fixes for Thin Translucent materials export.

#rnx

diff --git a/Engine/Plugins/Enterprise/GLTFExporter/Source/GLTFExporter/Private/Tasks/GLTFDelayedMaterialTasks.cpp b/Engine/Plugins/Enterprise/GLTFExporter/Source/GLTFExporter/Private/Tasks/GLTFDelayedMaterialTasks.cpp
--- a/Engine/Plugins/Enterprise/GLTFExporter/Source/GLTFExporter/Private/Tasks/GLTFDelayedMaterialTasks.cpp
+++ b/Engine/Plugins/Enterprise/GLTFExporter/Source/GLTFExporter/Private/Tasks/GLTFDelayedMaterialTasks.cpp
@@ -1139,15 +1139,21 @@ bool FGLTFDelayedMaterialTask::TryGetFuzzColorAndCloth(FGLTFJsonMaterial& OutMat
 bool FGLTFDelayedMaterialTask::TryGetTransmissionBaseColorAndOpacity(FGLTFJsonMaterial& OutMaterial, const FMaterialPropertyEx& BaseColorProperty, const FMaterialPropertyEx& OpacityProperty)
 {
 	const bool bIsBaseColorConstant = TryGetConstantColor(OutMaterial.PBRMetallicRoughness.BaseColorFactor, BaseColorProperty);
-	
+
 	float Opacity = 0.0f;
 	const bool bIsOpacityConstant = TryGetConstantScalar(Opacity, OpacityProperty);
 
-	if (bIsBaseColorConstant && bIsOpacityConstant)
+	if (bIsBaseColorConstant)
 	{
-		OutMaterial.Transmission.Factor = 1 - Opacity;
 		OutMaterial.PBRMetallicRoughness.BaseColorFactor.A = 1.0f;
+	}
+	if (bIsOpacityConstant)
+	{
+		OutMaterial.Transmission.Factor = 1 - Opacity;
+	}
 
+	if (bIsBaseColorConstant && bIsOpacityConstant)
+	{
 		return true;
 	}
 
@@ -1164,7 +1170,7 @@ bool FGLTFDelayedMaterialTask::TryGetTransmissionBaseColorAndOpacity(FGLTFJsonMa
 		if (!TryGetBakedMaterialProperty(OutMaterial.Transmission.Texture, OutMaterial.Transmission.Factor, 1.0f, OpacityProperty,
 			[](FColor& Color, const uint8& ChannelValueToMove)
 			{
-				Color.R = ChannelValueToMove;
+				Color.R = 255 - ChannelValueToMove;
 
 				Color.G = 255;
 				Color.B = 255;
