[GLTF Exporter] Extended analytics to also report Material types and supported materials

#rnx

diff --git a/Engine/Plugins/Enterprise/GLTFExporter/Source/GLTFExporter/Private/Builders/GLTFAnalyticsBuilder.cpp b/Engine/Plugins/Enterprise/GLTFExporter/Source/GLTFExporter/Private/Builders/GLTFAnalyticsBuilder.cpp
--- a/Engine/Plugins/Enterprise/GLTFExporter/Source/GLTFExporter/Private/Builders/GLTFAnalyticsBuilder.cpp
+++ b/Engine/Plugins/Enterprise/GLTFExporter/Source/GLTFExporter/Private/Builders/GLTFAnalyticsBuilder.cpp
@@ -5,6 +5,9 @@
 #include "StaticMeshResources.h"
 #include "Rendering/SkeletalMeshLODRenderData.h"
 #include "Converters/GLTFMeshUtilities.h"
+#include "Converters/GLTFMaterialUtilities.h"
+#include "Utilities/GLTFCoreUtilities.h"
+#include "MaterialShaderType.h"
 
 FGLTFAnalyticsBuilder::FGLTFAnalyticsBuilder(const FString& FileName, const UGLTFExportOptions* ExportOptions)
 	: FGLTFBufferBuilder(FileName, ExportOptions)
@@ -57,6 +60,34 @@ TArray<FAnalyticsEventAttribute> FGLTFAnalyticsBuilder::GenerateAnalytics() cons
 	}
 	EventAttributes.Emplace(TEXT("NumberOfVertices"), LexToString(NumberOfVertices));
 
+	//Material related Analytics:
+	{
+		TMap<EMaterialShadingModel, int32> NumberOfShadingModels;
+		for (const UMaterialInterface* MaterialInterface : MaterialsRecorded)
+		{
+			FString WarningMessage;
+			EMaterialShadingModel ShadingModel = FGLTFMaterialUtilities::GetShadingModel(MaterialInterface, WarningMessage);
+
+			int32& NumberOfShadingModel = NumberOfShadingModels.FindOrAdd(ShadingModel);
+			NumberOfShadingModel++;
+		}
+		TMap<FString, int32> NumberOfShadingModelsMap;
+		for (const TPair<EMaterialShadingModel, int32> NumberOfShadingModel : NumberOfShadingModels)
+		{
+			FString ShadingModelString = GetShadingModelString(NumberOfShadingModel.Key);
+			NumberOfShadingModelsMap.Add(ShadingModelString, NumberOfShadingModel.Value);
+		}
+		EventAttributes.Emplace(TEXT("NumberOfShadingModels"), NumberOfShadingModelsMap);
+
+		TArray<FString> ExporterSupportedShadingModels;
+		constexpr int32 NumMaterials = static_cast<int32>(EGLTFJsonShadingModel::NumShadingModels);
+		for (size_t MaterialIndex = 1; MaterialIndex < NumMaterials; MaterialIndex++)//skipping None
+		{
+			ExporterSupportedShadingModels.Add(FGLTFCoreUtilities::GetShadingModelString(EGLTFJsonShadingModel(MaterialIndex)));
+		}
+		EventAttributes.Emplace(TEXT("ExporterSupportedShadingModels"), ExporterSupportedShadingModels);
+	}
+
 	return EventAttributes;
 }
 
diff --git a/Engine/Plugins/Enterprise/GLTFExporter/Source/GLTFExporter/Private/Converters/GLTFMaterialUtilities.cpp b/Engine/Plugins/Enterprise/GLTFExporter/Source/GLTFExporter/Private/Converters/GLTFMaterialUtilities.cpp
--- a/Engine/Plugins/Enterprise/GLTFExporter/Source/GLTFExporter/Private/Converters/GLTFMaterialUtilities.cpp
+++ b/Engine/Plugins/Enterprise/GLTFExporter/Source/GLTFExporter/Private/Converters/GLTFMaterialUtilities.cpp
@@ -524,3 +524,74 @@ bool FGLTFMaterialUtilities::NeedsMeshData(const TArray<const UMaterialInterface
 
 	return false;
 }
+
+EMaterialShadingModel FGLTFMaterialUtilities::GetShadingModel(const UMaterialInterface* Material, FString& OutMessage)
+{
+	const FMaterialShadingModelField Possibilities = Material->GetShadingModels();
+	const int32 PossibilitiesCount = Possibilities.CountShadingModels();
+
+	if (PossibilitiesCount == 0)
+	{
+		const EMaterialShadingModel ShadingModel = MSM_DefaultLit;
+		OutMessage = FString::Printf(
+			TEXT("No shading model defined for material %s, will export as %s"),
+			*Material->GetName(),
+			*FGLTFNameUtilities::GetName(ShadingModel));
+		return ShadingModel;
+	}
+
+	if (PossibilitiesCount > 1)
+	{
+#if WITH_EDITOR
+		if (!FApp::CanEverRender())
+		{
+			const EMaterialShadingModel ShadingModel = FGLTFMaterialUtilities::GetRichestShadingModel(Possibilities);
+			OutMessage = FString::Printf(
+				TEXT("Can't evaluate shading model expression in material %s because renderer missing, will export as %s"),
+				*Material->GetName(),
+				*FGLTFNameUtilities::GetName(ShadingModel));
+			return ShadingModel;
+		}
+
+		if (Material->IsShadingModelFromMaterialExpression())
+		{
+			const FMaterialShadingModelField Evaluation = FGLTFMaterialUtilities::EvaluateShadingModelExpression(Material);
+			const int32 EvaluationCount = Evaluation.CountShadingModels();
+
+			if (EvaluationCount == 0)
+			{
+				const EMaterialShadingModel ShadingModel = FGLTFMaterialUtilities::GetRichestShadingModel(Possibilities);
+				OutMessage = FString::Printf(
+					TEXT("Evaluation of shading model expression in material %s returned none, will export as %s"),
+					*Material->GetName(),
+					*FGLTFNameUtilities::GetName(ShadingModel));
+				return ShadingModel;
+			}
+
+			if (EvaluationCount > 1)
+			{
+				const EMaterialShadingModel ShadingModel = FGLTFMaterialUtilities::GetRichestShadingModel(Evaluation);
+				OutMessage = FString::Printf(
+					TEXT("Evaluation of shading model expression in material %s is inconclusive (%s), will export as %s"),
+					*Material->GetName(),
+					*FGLTFMaterialUtilities::ShadingModelsToString(Evaluation),
+					*FGLTFNameUtilities::GetName(ShadingModel));
+				return ShadingModel;
+			}
+
+			return Evaluation.GetFirstShadingModel();
+		}
+
+		// we should never end up here
+#else
+		const EMaterialShadingModel ShadingModel = FGLTFMaterialUtilities::GetRichestShadingModel(Possibilities);
+		OutMessage = FString::Printf(
+			TEXT("Can't evaluate shading model expression in material %s without editor, will export as %s"),
+			*Material->GetName(),
+			*FGLTFNameUtilities::GetName(ShadingModel));
+		return ShadingModel;
+#endif
+	}
+
+	return Possibilities.GetFirstShadingModel();
+}
diff --git a/Engine/Plugins/Enterprise/GLTFExporter/Source/GLTFExporter/Private/Converters/GLTFMaterialUtilities.h b/Engine/Plugins/Enterprise/GLTFExporter/Source/GLTFExporter/Private/Converters/GLTFMaterialUtilities.h
--- a/Engine/Plugins/Enterprise/GLTFExporter/Source/GLTFExporter/Private/Converters/GLTFMaterialUtilities.h
+++ b/Engine/Plugins/Enterprise/GLTFExporter/Source/GLTFExporter/Private/Converters/GLTFMaterialUtilities.h
@@ -47,6 +47,8 @@ struct FGLTFMaterialUtilities
 	static bool NeedsMeshData(const UMaterialInterface* Material);
 	static bool NeedsMeshData(const TArray<const UMaterialInterface*>& Materials);
 
+	static EMaterialShadingModel GetShadingModel(const UMaterialInterface* Material, FString& OutMessage);
+
 #if WITH_EDITOR
 	static bool IsNormalMap(const FMaterialPropertyEx& Property);
 	static bool IsSRGB(const FMaterialPropertyEx& Property);
@@ -86,7 +88,6 @@ struct FGLTFMaterialUtilities
 	static void AnalyzeMaterialProperty(const UMaterialInterface* Material, const FMaterialPropertyEx& InProperty, FMaterialAnalysisResult& OutAnalysis);
 
 	static FMaterialShadingModelField EvaluateShadingModelExpression(const UMaterialInterface* Material);
-
 private:
 
 	template<typename ExpressionType>
diff --git a/Engine/Plugins/Enterprise/GLTFExporter/Source/GLTFExporter/Private/Tasks/GLTFDelayedMaterialTasks.cpp b/Engine/Plugins/Enterprise/GLTFExporter/Source/GLTFExporter/Private/Tasks/GLTFDelayedMaterialTasks.cpp
--- a/Engine/Plugins/Enterprise/GLTFExporter/Source/GLTFExporter/Private/Tasks/GLTFDelayedMaterialTasks.cpp
+++ b/Engine/Plugins/Enterprise/GLTFExporter/Source/GLTFExporter/Private/Tasks/GLTFDelayedMaterialTasks.cpp
@@ -406,81 +406,16 @@ void FGLTFDelayedMaterialTask::GetProxyParameter(const FGLTFProxyMaterialTexture
 	}
 }
 
-EMaterialShadingModel FGLTFDelayedMaterialTask::GetShadingModel() const
+void FGLTFDelayedMaterialTask::ConvertShadingModel(EGLTFJsonShadingModel& OutShadingModel) const
 {
-	const FMaterialShadingModelField Possibilities = Material->GetShadingModels();
-	const int32 PossibilitiesCount = Possibilities.CountShadingModels();
+	FString WarningMessage;
+	EMaterialShadingModel ShadingModel = FGLTFMaterialUtilities::GetShadingModel(Material, WarningMessage);
 
-	if (PossibilitiesCount == 0)
+	if (!WarningMessage.IsEmpty())
 	{
-		const EMaterialShadingModel ShadingModel = MSM_DefaultLit;
-		Builder.LogWarning(FString::Printf(
-			TEXT("No shading model defined for material %s, will export as %s"),
-			*Material->GetName(),
-			*FGLTFNameUtilities::GetName(ShadingModel)));
-		return ShadingModel;
+		Builder.LogWarning(WarningMessage);
 	}
 
-	if (PossibilitiesCount > 1)
-	{
-#if WITH_EDITOR
-		if (!FApp::CanEverRender())
-		{
-			const EMaterialShadingModel ShadingModel = FGLTFMaterialUtilities::GetRichestShadingModel(Possibilities);
-			Builder.LogWarning(FString::Printf(
-				TEXT("Can't evaluate shading model expression in material %s because renderer missing, will export as %s"),
-				*Material->GetName(),
-				*FGLTFNameUtilities::GetName(ShadingModel)));
-			return ShadingModel;
-		}
-
-		if (Material->IsShadingModelFromMaterialExpression())
-		{
-			const FMaterialShadingModelField Evaluation = FGLTFMaterialUtilities::EvaluateShadingModelExpression(Material);
-			const int32 EvaluationCount = Evaluation.CountShadingModels();
-
-			if (EvaluationCount == 0)
-			{
-				const EMaterialShadingModel ShadingModel = FGLTFMaterialUtilities::GetRichestShadingModel(Possibilities);
-				Builder.LogWarning(FString::Printf(
-					TEXT("Evaluation of shading model expression in material %s returned none, will export as %s"),
-					*Material->GetName(),
-					*FGLTFNameUtilities::GetName(ShadingModel)));
-				return ShadingModel;
-			}
-
-			if (EvaluationCount > 1)
-			{
-				const EMaterialShadingModel ShadingModel = FGLTFMaterialUtilities::GetRichestShadingModel(Evaluation);
-				Builder.LogWarning(FString::Printf(
-					TEXT("Evaluation of shading model expression in material %s is inconclusive (%s), will export as %s"),
-					*Material->GetName(),
-					*FGLTFMaterialUtilities::ShadingModelsToString(Evaluation),
-					*FGLTFNameUtilities::GetName(ShadingModel)));
-				return ShadingModel;
-			}
-
-			return Evaluation.GetFirstShadingModel();
-		}
-
-		// we should never end up here
-#else
-		const EMaterialShadingModel ShadingModel = FGLTFMaterialUtilities::GetRichestShadingModel(Possibilities);
-		Builder.LogWarning(FString::Printf(
-			TEXT("Can't evaluate shading model expression in material %s without editor, will export as %s"),
-			*Material->GetName(),
-			*FGLTFNameUtilities::GetName(ShadingModel)));
-		return ShadingModel;
-#endif
-	}
-
-	return Possibilities.GetFirstShadingModel();
-}
-
-void FGLTFDelayedMaterialTask::ConvertShadingModel(EGLTFJsonShadingModel& OutShadingModel) const
-{
-	EMaterialShadingModel ShadingModel = GetShadingModel();
-
 	const EBlendMode BlendMode = Material->GetBlendMode();
 	if (ShadingModel == MSM_ClearCoat && BlendMode != BLEND_Opaque && BlendMode != BLEND_Masked)
 	{
diff --git a/Engine/Plugins/Enterprise/GLTFExporter/Source/GLTFExporter/Private/Tasks/GLTFDelayedMaterialTasks.h b/Engine/Plugins/Enterprise/GLTFExporter/Source/GLTFExporter/Private/Tasks/GLTFDelayedMaterialTasks.h
--- a/Engine/Plugins/Enterprise/GLTFExporter/Source/GLTFExporter/Private/Tasks/GLTFDelayedMaterialTasks.h
+++ b/Engine/Plugins/Enterprise/GLTFExporter/Source/GLTFExporter/Private/Tasks/GLTFDelayedMaterialTasks.h
@@ -64,7 +64,6 @@ private:
 		return ParameterInfo.Get(Material, Value, true);
 	}
 
-	EMaterialShadingModel GetShadingModel() const;
 	void ConvertShadingModel(EGLTFJsonShadingModel& OutShadingModel) const;
 	void ConvertAlphaMode(EGLTFJsonAlphaMode& OutAlphaMode) const;
 
diff --git a/Engine/Plugins/Enterprise/GLTFExporter/Source/GLTFExporter/Private/Utilities/GLTFCoreUtilities.cpp b/Engine/Plugins/Enterprise/GLTFExporter/Source/GLTFExporter/Private/Utilities/GLTFCoreUtilities.cpp
--- a/Engine/Plugins/Enterprise/GLTFExporter/Source/GLTFExporter/Private/Utilities/GLTFCoreUtilities.cpp
+++ b/Engine/Plugins/Enterprise/GLTFExporter/Source/GLTFExporter/Private/Utilities/GLTFCoreUtilities.cpp
@@ -228,6 +228,21 @@ EGLTFJsonShadingModel FGLTFCoreUtilities::ConvertShadingModel(EMaterialShadingMo
 	}
 }
 
+FString FGLTFCoreUtilities::GetShadingModelString(EGLTFJsonShadingModel ShadingModel)
+{
+	switch (ShadingModel)
+	{
+		case EGLTFJsonShadingModel::None:               return TEXT("Unknown");
+		case EGLTFJsonShadingModel::Default:            return TEXT("Default");
+		case EGLTFJsonShadingModel::Unlit:              return TEXT("Unlit");
+		case EGLTFJsonShadingModel::ClearCoat:          return TEXT("ClearCoat");
+		case EGLTFJsonShadingModel::Sheen:              return TEXT("Cloth");
+		case EGLTFJsonShadingModel::Transmission:       return TEXT("ThinTranslucent");
+		case EGLTFJsonShadingModel::NumShadingModels:   return TEXT("Unknown");
+		default:                                        return TEXT("Unknown");
+	}
+}
+
 EGLTFJsonAlphaMode FGLTFCoreUtilities::ConvertAlphaMode(EBlendMode Mode)
 {
 	switch (Mode)
diff --git a/Engine/Plugins/Enterprise/GLTFExporter/Source/GLTFExporter/Public/Utilities/GLTFCoreUtilities.h b/Engine/Plugins/Enterprise/GLTFExporter/Source/GLTFExporter/Public/Utilities/GLTFCoreUtilities.h
--- a/Engine/Plugins/Enterprise/GLTFExporter/Source/GLTFExporter/Public/Utilities/GLTFCoreUtilities.h
+++ b/Engine/Plugins/Enterprise/GLTFExporter/Source/GLTFExporter/Public/Utilities/GLTFCoreUtilities.h
@@ -60,6 +60,7 @@ struct GLTFEXPORTER_API FGLTFCoreUtilities
 	static EGLTFJsonInterpolation ConvertInterpolation(const EAnimInterpolationType Type);
 
 	static EGLTFJsonShadingModel ConvertShadingModel(EMaterialShadingModel ShadingModel);
+	static FString GetShadingModelString(EGLTFJsonShadingModel ShadingModel);
 
 	static EGLTFJsonAlphaMode ConvertAlphaMode(EBlendMode Mode);
 
