Removed glTF Proxy material feature

#rnx

diff --git a/Engine/Plugins/Enterprise/GLTFExporter/Source/GLTFExporter/Private/Actions/GLTFProxyAssetActions.cpp b/Engine/Plugins/Enterprise/GLTFExporter/Source/GLTFExporter/Private/Actions/GLTFProxyAssetActions.cpp
--- a/Engine/Plugins/Enterprise/GLTFExporter/Source/GLTFExporter/Private/Actions/GLTFProxyAssetActions.cpp
+++ b/Engine/Plugins/Enterprise/GLTFExporter/Source/GLTFExporter/Private/Actions/GLTFProxyAssetActions.cpp
@@ -4,10 +4,12 @@
 
 #include "Actions/GLTFProxyAssetActions.h"
 #include "Actions/GLTFEditorStyle.h"
+#include "AssetToolsModule.h"
 #include "Materials/MaterialInterface.h"
 #include "Options/GLTFProxyOptions.h"
 #include "UI/GLTFProxyOptionsWindow.h"
 #include "Converters/GLTFMaterialProxyFactory.h"
+#include "UserData/GLTFMaterialUserData.h"
 #include "ContentBrowserMenuContexts.h"
 #include "ToolMenus.h"
 
@@ -24,6 +26,17 @@ void FGLTFProxyAssetActions::AddMenuEntry()
 	{
 		return;
 	}
+	
+	// Do not add any entry if UGLTFMaterialExportOptions is not supported
+	IAssetTools& AssetTools = FAssetToolsModule::GetModule().Get();
+	TSharedPtr<FPathPermissionList> AssetClassPermissionList = AssetTools.GetAssetClassPathPermissionList(EAssetClassAction::CreateAsset);
+	if (AssetClassPermissionList && AssetClassPermissionList->HasFiltering())
+	{
+		if (!AssetClassPermissionList->PassesFilter(UGLTFMaterialExportOptions::StaticClass()->GetPathName()))
+		{
+			return;
+		}
+	}
 
 	UToolMenu* Menu = ToolMenus->ExtendMenu(MenuName);
 	FToolMenuSection& Section = Menu->FindOrAddSection(SectionName);
