[glTFExporter] Improved glTF exporter's material expression matching to support same expressions generated by the Importer

#rnx

diff --git a/Engine/Plugins/Enterprise/GLTFExporter/GLTFExporter.uplugin b/Engine/Plugins/Enterprise/GLTFExporter/GLTFExporter.uplugin
--- a/Engine/Plugins/Enterprise/GLTFExporter/GLTFExporter.uplugin
+++ b/Engine/Plugins/Enterprise/GLTFExporter/GLTFExporter.uplugin
@@ -11,6 +11,11 @@
 	"SupportURL": "",
 	"EnabledByDefault": true,
 	"CanContainContent": true,
+	"SupportedTargetPlatforms": [
+		"Win64",
+		"Linux",
+		"Mac"
+	],
 	"Modules": [
 		{
 			"Name": "GLTFExporter",
@@ -23,6 +28,11 @@
 		{
 			"Name": "VariantManagerContent",
 			"Enabled": true
+		},
+		{
+			"Name": "Interchange",
+			"Enabled": true,
+			"Optional": true
 		}
 	]
 }
diff --git a/Engine/Plugins/Enterprise/GLTFExporter/Source/GLTFExporter/GLTFExporter.Build.cs b/Engine/Plugins/Enterprise/GLTFExporter/Source/GLTFExporter/GLTFExporter.Build.cs
--- a/Engine/Plugins/Enterprise/GLTFExporter/Source/GLTFExporter/GLTFExporter.Build.cs
+++ b/Engine/Plugins/Enterprise/GLTFExporter/Source/GLTFExporter/GLTFExporter.Build.cs
@@ -16,6 +16,13 @@ public class GLTFExporter : ModuleRules
 				"Engine",
 			}
 		);
+		
+		PrivateIncludePathModuleNames.AddRange(
+			new string[]
+			{
+				"InterchangeImport",
+			}
+		);
 
 		PrivateDependencyModuleNames.AddRange(
 			new string[]
diff --git a/Engine/Plugins/Enterprise/GLTFExporter/Source/GLTFExporter/Private/Builders/GLTFConvertBuilder.cpp b/Engine/Plugins/Enterprise/GLTFExporter/Source/GLTFExporter/Private/Builders/GLTFConvertBuilder.cpp
--- a/Engine/Plugins/Enterprise/GLTFExporter/Source/GLTFExporter/Private/Builders/GLTFConvertBuilder.cpp
+++ b/Engine/Plugins/Enterprise/GLTFExporter/Source/GLTFExporter/Private/Builders/GLTFConvertBuilder.cpp
@@ -289,22 +289,22 @@ FGLTFJsonTexture* FGLTFConvertBuilder::AddUniqueTexture(const UTextureRenderTarg
 	return AddUniqueTexture(Texture, Texture->SRGB);
 }
 
-FGLTFJsonTexture* FGLTFConvertBuilder::AddUniqueTexture(const UTexture* Texture, bool bToSRGB)
+FGLTFJsonTexture* FGLTFConvertBuilder::AddUniqueTexture(const UTexture* Texture, bool bToSRGB, TextureAddress TextureAddressX, TextureAddress TextureAddressY)
 {
 	if (const UTexture2D* Texture2D = Cast<UTexture2D>(Texture))
 	{
-		return AddUniqueTexture(Texture2D, bToSRGB);
+		return AddUniqueTexture(Texture2D, bToSRGB, TextureAddressX, TextureAddressY);
 	}
 
 	if (const UTextureRenderTarget2D* RenderTarget2D = Cast<UTextureRenderTarget2D>(Texture))
 	{
-		return AddUniqueTexture(RenderTarget2D, bToSRGB);
+		return AddUniqueTexture(RenderTarget2D, bToSRGB, TextureAddressX, TextureAddressY);
 	}
 
 	return nullptr;
 }
 
-FGLTFJsonTexture* FGLTFConvertBuilder::AddUniqueTexture(const UTexture2D* Texture, bool bToSRGB)
+FGLTFJsonTexture* FGLTFConvertBuilder::AddUniqueTexture(const UTexture2D* Texture, bool bToSRGB, TextureAddress TextureAddressX, TextureAddress TextureAddressY)
 {
 	if (Texture == nullptr)
 	{
@@ -313,7 +313,7 @@ FGLTFJsonTexture* FGLTFConvertBuilder::AddUniqueTexture(const UTexture2D* Textur
 
 	RecordTexture(Texture);
 
-	return Texture2DConverter->GetOrAdd(Texture, bToSRGB);
+	return Texture2DConverter->GetOrAdd(Texture, bToSRGB, TextureAddressX, TextureAddressY);
 }
 
 FGLTFJsonTexture* FGLTFConvertBuilder::AddUniqueTexture(const UTextureRenderTarget2D* Texture, bool bToSRGB)
diff --git a/Engine/Plugins/Enterprise/GLTFExporter/Source/GLTFExporter/Private/Converters/GLTFMaterialProxyFactory.cpp b/Engine/Plugins/Enterprise/GLTFExporter/Source/GLTFExporter/Private/Converters/GLTFMaterialProxyFactory.cpp
--- a/Engine/Plugins/Enterprise/GLTFExporter/Source/GLTFExporter/Private/Converters/GLTFMaterialProxyFactory.cpp
+++ b/Engine/Plugins/Enterprise/GLTFExporter/Source/GLTFExporter/Private/Converters/GLTFMaterialProxyFactory.cpp
@@ -264,12 +264,12 @@ TUniquePtr<IGLTFTexture2DConverter> FGLTFMaterialProxyFactory::CreateTextureConv
 
 	protected:
 
-		virtual void Sanitize(const UTexture2D*& Texture2D, bool& bToSRGB) override
+		virtual void Sanitize(const UTexture2D*& Texture2D, bool& bToSRGB, TextureAddress& WrapS, TextureAddress& WrapT) override
 		{
 			bToSRGB = false; // ignore
 		}
 
-		virtual FGLTFJsonTexture* Convert(const UTexture2D* Texture2D, bool bToSRGB) override
+		virtual FGLTFJsonTexture* Convert(const UTexture2D* Texture2D, bool bToSRGB, TextureAddress WrapS, TextureAddress WrapT) override
 		{
 			FGLTFJsonTexture* Texture = Factory.Builder.AddTexture();
 			Factory.Textures.Add(Texture, const_cast<UTexture2D*>(Texture2D));
diff --git a/Engine/Plugins/Enterprise/GLTFExporter/Source/GLTFExporter/Private/Converters/GLTFMaterialUtilities.cpp b/Engine/Plugins/Enterprise/GLTFExporter/Source/GLTFExporter/Private/Converters/GLTFMaterialUtilities.cpp
--- a/Engine/Plugins/Enterprise/GLTFExporter/Source/GLTFExporter/Private/Converters/GLTFMaterialUtilities.cpp
+++ b/Engine/Plugins/Enterprise/GLTFExporter/Source/GLTFExporter/Private/Converters/GLTFMaterialUtilities.cpp
@@ -16,9 +16,21 @@
 #include "Materials/MaterialAttributeDefinitionMap.h"
 #include "Materials/MaterialExpressionCustomOutput.h"
 #include "Materials/MaterialExpressionTextureCoordinate.h"
+#endif
+
+#include "Materials/MaterialInstance.h"
+#include "Materials/MaterialExpressionMaterialFunctionCall.h"
+#include "Materials/MaterialExpressionConstant.h"
+#include "Materials/MaterialExpressionConstant2Vector.h"
+#include "Materials/MaterialExpressionConstant3Vector.h"
+#include "Materials/MaterialExpressionConstant4Vector.h"
+#include "Materials/MaterialExpressionScalarParameter.h"
+#include "Materials/MaterialExpressionVectorParameter.h"
 #include "Materials/MaterialExpressionTextureSample.h"
 #include "Materials/MaterialExpressionTextureSampleParameter2D.h"
-#endif
+
+#include "Utilities/GLTFCoreUtilities.h"
+#include "Gltf/InterchangeGLTFMaterial.h"
 
 UMaterialInterface* FGLTFMaterialUtilities::GetDefaultMaterial()
 {
@@ -595,3 +607,426 @@ EMaterialShadingModel FGLTFMaterialUtilities::GetShadingModel(const UMaterialInt
 
 	return Possibilities.GetFirstShadingModel();
 }
+
+FGLTFImportMaterialMatchMakingHelper::FGLTFImportMaterialMatchMakingHelper(FGLTFConvertBuilder& InBuilder,
+	const UMaterialInterface* InMaterial,
+	FGLTFJsonMaterial& InJsonMaterial)
+	: Builder(InBuilder)
+	, JsonMaterial(InJsonMaterial)
+	, Material(InMaterial)
+	, bMaterialInstance(InMaterial ? InMaterial->IsA(UMaterialInstance::StaticClass()) : false)
+	, bIsGLTFImportedMaterial(false)
+{
+	if (!Material)
+	{
+		return;
+	}
+
+	using namespace UE::Interchange::GLTFMaterials;
+
+	const UMaterial* TopMostParentMaterial = Material->GetMaterial();
+
+	if (bMaterialInstance)
+	{
+		FString PathName = TopMostParentMaterial->GetPathName();
+
+		TMap<FString, EShadingModel> MaterialPathsToShadingModels = GetMaterialPathsToShadingModels();
+
+		if (MaterialPathsToShadingModels.Contains(PathName))
+		{
+			bIsGLTFImportedMaterial = true;
+
+			JsonMaterial.ShadingModel = EGLTFJsonShadingModel(MaterialPathsToShadingModels[PathName]);
+		}
+	}
+	else
+	{
+#if WITH_EDITOR
+		//Acquire glTF Importer's MaterialFunctionPaths
+		TMap<FString, EShadingModel> MaterialFunctionPathsToShadingModels = GetMaterialFunctionPathsToShadingModels();
+
+		//Acquire the Material's expressions and check if any matches a MaterialFunctionCall which matches the glTF Importer's list:
+		TConstArrayView<TObjectPtr<UMaterialExpression>> MaterialExpressions = TopMostParentMaterial->GetExpressions();
+
+		UMaterialExpressionMaterialFunctionCall* GLTFImportedMaterialFunction = nullptr;
+		EShadingModel GLTFImportedShadingModel = EShadingModel::DEFAULT;
+
+		for (const TObjectPtr<UMaterialExpression>& MaterialExpression : MaterialExpressions)
+		{
+			if (UMaterialExpressionMaterialFunctionCall* FunctionCallExpression = Cast<UMaterialExpressionMaterialFunctionCall>(MaterialExpression))
+			{
+				if (FunctionCallExpression->MaterialFunction)
+				{
+					FString PathName = FunctionCallExpression->MaterialFunction->GetPathName();
+
+					if (MaterialFunctionPathsToShadingModels.Contains(PathName))
+					{
+						//Material imported by glTF Importer
+						GLTFImportedMaterialFunction = FunctionCallExpression;
+						GLTFImportedShadingModel = MaterialFunctionPathsToShadingModels[PathName];
+
+						break;
+					}
+				}
+			}
+		}
+
+		if (!GLTFImportedMaterialFunction)
+		{
+			return;
+		}
+
+		bIsGLTFImportedMaterial = true;
+
+		JsonMaterial.ShadingModel = EGLTFJsonShadingModel(GLTFImportedShadingModel);
+
+		//If Exported Material is a MaterialInstance then we will use the GetParameter APIs
+		if (!bMaterialInstance)
+		{
+			TArrayView<FExpressionInput*> InputsView = GLTFImportedMaterialFunction->GetInputsView();
+
+			for (FExpressionInput* Input : InputsView)
+			{
+				if (!Input || !Input->Expression || !Input->IsConnected())
+				{
+					continue;
+				}
+
+				FString Name = Input->InputName.ToString();
+				UMaterialExpression* InputExpression = Input->GetTracedInput().Expression;
+
+				Inputs.Add(Name, InputExpression);
+			}
+		}
+#endif
+	}
+}
+
+bool FGLTFImportMaterialMatchMakingHelper::GetValue(const FString& InputKey, float& OutValue)
+{
+	if (bMaterialInstance)
+	{
+		if (!Material->GetScalarParameterValue(*InputKey, OutValue, true))
+		{
+			return false;
+		}
+	}
+	else
+	{
+		if (!Inputs.Contains(InputKey))
+		{
+			return false;
+		}
+
+		UMaterialExpression* InputExpression = Inputs[InputKey];
+
+		if (UMaterialExpressionConstant* ConstExpr1 = Cast<UMaterialExpressionConstant>(InputExpression))
+		{
+			OutValue = ConstExpr1->R;
+		}
+		else
+		{
+			return false;
+		}
+	}
+
+	return true;
+}
+
+bool FGLTFImportMaterialMatchMakingHelper::GetValue(const FString& InputKey, FGLTFJsonColor3& OutValue)
+{
+	FLinearColor Value;
+
+	if (bMaterialInstance)
+	{
+		if (!Material->GetVectorParameterValue(*InputKey, Value, true))
+		{
+			return false;
+		}
+	}
+	else
+	{
+		if (!Inputs.Contains(InputKey))
+		{
+			return false;
+		}
+
+		UMaterialExpression* InputExpression = Inputs[InputKey];
+
+		if (UMaterialExpressionConstant3Vector* ConstExpr3 = Cast<UMaterialExpressionConstant3Vector>(InputExpression))
+		{
+			Value = ConstExpr3->Constant;
+		}
+		else
+		{
+			return false;
+		}
+	}
+
+	OutValue = FGLTFCoreUtilities::ConvertColor3(Value);
+
+	return true;
+}
+
+bool FGLTFImportMaterialMatchMakingHelper::GetValue(const FString& InputKey, FGLTFJsonColor4& OutValue, bool HandleAsColor)
+{
+	using namespace UE::Interchange::GLTFMaterials::Inputs::PostFix;
+
+	FLinearColor Value;
+
+	if (bMaterialInstance)
+	{
+		if (!Material->GetVectorParameterValue(*InputKey, Value, true))
+		{
+			return false;
+		}
+	}
+	else
+	{
+		if (!Inputs.Contains(InputKey))
+		{
+			if (HandleAsColor)
+			{
+				float Alpha = 1;
+				FGLTFJsonColor3 RGB = FGLTFJsonColor3::White;
+				if (GetValue(InputKey + Color_RGB, RGB) || GetValue(InputKey + Color_A, Alpha))
+				{
+					OutValue = { RGB.R, RGB.G, RGB.B, FMath::Clamp(Alpha, 0.0f, 1.0f) };
+					return true;
+				}
+			}
+
+			return false;
+		}
+
+		UMaterialExpression* InputExpression = Inputs[InputKey];
+
+
+		if (UMaterialExpressionConstant4Vector* ConstExpr4 = Cast<UMaterialExpressionConstant4Vector>(InputExpression))
+		{
+			Value = ConstExpr4->Constant;
+		}
+		else
+		{
+			return false;
+		}
+	}
+
+	if (HandleAsColor)
+	{
+		OutValue = FGLTFCoreUtilities::ConvertColor(Value);
+	}
+	else
+	{
+		OutValue = { Value.R, Value.G, Value.B, Value.A };
+	}
+
+
+	return true;
+}
+
+//To acquire TextureAddress from glTF Imported material TilingMethod values:
+TextureAddress GetTextureAddress(float Value)
+{
+	//0: wrap
+	//1: clamp
+	//2: mirror
+	int ValueInt = FMath::RoundToInt(Value);
+
+	switch (ValueInt)
+	{
+		case 0:     return TextureAddress::TA_Wrap;
+		case 1:     return TextureAddress::TA_Clamp;
+		case 2:	    return TextureAddress::TA_Mirror;
+		default:    return TextureAddress::TA_MAX;
+	}
+}
+
+bool FGLTFImportMaterialMatchMakingHelper::GetValue(const FString& InputKey, FGLTFJsonTextureInfo& OutValue)
+{
+	using namespace UE::Interchange::GLTFMaterials::Inputs;
+
+	UTexture* Texture = nullptr;
+
+	if (bMaterialInstance)
+	{
+		if (!Material->GetTextureParameterValue(*InputKey, Texture, true))
+		{
+			return false;
+		}
+	}
+	else
+	{
+
+		if (!Inputs.Contains(InputKey))
+		{
+			return false;
+		}
+
+		UMaterialExpression* InputExpression = Inputs[InputKey];
+
+		if (UMaterialExpressionTextureObject* TextureObject = Cast<UMaterialExpressionTextureObject>(InputExpression))
+		{
+			if (TextureObject->Texture)
+			{
+				Texture = TextureObject->Texture.Get();
+			}
+		}
+		else
+		{
+			return false;
+		}
+	}
+
+	if (!Texture)
+	{
+		return false;
+	}
+
+	TextureAddress TextureAddressX = TextureAddress::TA_Wrap; //glTF Importer's default tiling method is Wrap
+	TextureAddress TextureAddressY = TextureAddress::TA_Wrap; //glTF Importer's default tiling method is Wrap
+
+	FGLTFJsonColor4 TilingMethod = FGLTFColor4{ 0.f, 0.f, 0.f, 0.f };
+	if (GetValue(InputKey + PostFix::TilingMethod, TilingMethod, false))
+	{
+		TextureAddressX = GetTextureAddress(TilingMethod.R);
+		TextureAddressY = GetTextureAddress(TilingMethod.G);
+	}
+
+	const bool bSRGB = InputKey == BaseColorTexture || InputKey == EmissiveTexture;
+	OutValue.Index = Builder.AddUniqueTexture(Texture, bSRGB, TextureAddressX, TextureAddressY);
+
+	float UVIndex;
+	if (GetValue(InputKey + PostFix::TexCoord, UVIndex))
+	{
+		OutValue.TexCoord = FMath::RoundToInt(UVIndex);
+	}
+
+	FGLTFJsonColor4 UVOffsetScale = FGLTFColor4{ 0.f, 0.f, 1.f, 1.f };
+	if (GetValue(InputKey + PostFix::OffsetScale, UVOffsetScale, false))
+	{
+		OutValue.Transform.Offset.X = UVOffsetScale.R;
+		OutValue.Transform.Offset.Y = UVOffsetScale.G;
+
+		OutValue.Transform.Scale.X = UVOffsetScale.B;
+		OutValue.Transform.Scale.Y = UVOffsetScale.A;
+	}
+	else
+	{
+		if (GetValue(InputKey + PostFix::OffsetX, UVOffsetScale.R))
+		{
+			OutValue.Transform.Offset.X = UVOffsetScale.R;
+		}
+		if (GetValue(InputKey + PostFix::OffsetY, UVOffsetScale.G))
+		{
+			OutValue.Transform.Offset.Y = UVOffsetScale.G;
+		}
+
+		if (GetValue(InputKey + PostFix::ScaleX, UVOffsetScale.B))
+		{
+			OutValue.Transform.Scale.X = UVOffsetScale.B;
+		}
+		if (GetValue(InputKey + PostFix::ScaleY, UVOffsetScale.A))
+		{
+			OutValue.Transform.Scale.Y = UVOffsetScale.A;
+		}
+	}
+
+	float UVRotation;
+	if (GetValue(InputKey + PostFix::Rotation, UVRotation))
+	{
+		OutValue.Transform.Rotation = UVRotation;
+	}
+
+	if (!Builder.ExportOptions->bExportTextureTransforms && !OutValue.Transform.IsExactlyDefault())
+	{
+		Builder.LogWarning(FString::Printf(
+			TEXT("Texture coordinates [%d] in %s for material %s are transformed, but texture transform is disabled by export options"),
+			OutValue.TexCoord,
+			*InputKey,
+			*Material->GetName()));
+		OutValue.Transform = {};
+	}
+
+	return true;
+}
+
+void FGLTFImportMaterialMatchMakingHelper::Process()
+{
+	using namespace UE::Interchange::GLTFMaterials;
+
+	GetValue(Inputs::AlphaCutoff, JsonMaterial.AlphaCutoff);
+
+	if (JsonMaterial.ShadingModel == EGLTFJsonShadingModel::SpecularGlossiness)
+	{
+		GetValue(Inputs::DiffuseTexture, JsonMaterial.PBRSpecularGlossiness.DiffuseTexture);
+		GetValue(Inputs::DiffuseFactor, JsonMaterial.PBRSpecularGlossiness.DiffuseFactor, true);
+	}
+	else
+	{
+		GetValue(Inputs::BaseColorTexture, JsonMaterial.PBRMetallicRoughness.BaseColorTexture);
+		GetValue(Inputs::BaseColorFactor, JsonMaterial.PBRMetallicRoughness.BaseColorFactor, true);
+	}
+
+	if (JsonMaterial.ShadingModel != EGLTFJsonShadingModel::Unlit)
+	{
+		if (JsonMaterial.ShadingModel != EGLTFJsonShadingModel::SpecularGlossiness)
+		{
+			GetValue(Inputs::MetallicRoughnessTexture, JsonMaterial.PBRMetallicRoughness.MetallicRoughnessTexture);
+			GetValue(Inputs::MetallicFactor, JsonMaterial.PBRMetallicRoughness.MetallicFactor);
+			GetValue(Inputs::RoughnessFactor, JsonMaterial.PBRMetallicRoughness.RoughnessFactor);
+
+			GetValue(Inputs::SpecularFactor, JsonMaterial.Specular.Factor);
+			GetValue(Inputs::SpecularTexture, JsonMaterial.Specular.Texture);
+
+			GetValue(Inputs::IOR, JsonMaterial.IOR.Value);
+		}
+
+		//Importer does not support Emissive for Transmission at the moment
+		GetValue(Inputs::EmissiveTexture, JsonMaterial.EmissiveTexture);
+		GetValue(Inputs::EmissiveFactor, JsonMaterial.EmissiveFactor);
+		GetValue(Inputs::EmissiveStrength, JsonMaterial.EmissiveStrength);
+
+		GetValue(Inputs::NormalTexture, JsonMaterial.NormalTexture);
+		GetValue(Inputs::NormalScale, JsonMaterial.NormalTexture.Scale);
+
+		GetValue(Inputs::OcclusionTexture, JsonMaterial.OcclusionTexture);
+		GetValue(Inputs::OcclusionStrength, JsonMaterial.OcclusionTexture.Strength);
+
+		if (JsonMaterial.ShadingModel == EGLTFJsonShadingModel::ClearCoat)
+		{
+			GetValue(Inputs::ClearCoatTexture, JsonMaterial.ClearCoat.ClearCoatTexture);
+			GetValue(Inputs::ClearCoatFactor, JsonMaterial.ClearCoat.ClearCoatFactor);
+
+			GetValue(Inputs::ClearCoatRoughnessTexture, JsonMaterial.ClearCoat.ClearCoatRoughnessTexture);
+			GetValue(Inputs::ClearCoatRoughnessFactor, JsonMaterial.ClearCoat.ClearCoatRoughnessFactor);
+
+			GetValue(Inputs::ClearCoatNormalTexture, JsonMaterial.ClearCoat.ClearCoatNormalTexture);
+			GetValue(Inputs::ClearCoatNormalScale, JsonMaterial.ClearCoat.ClearCoatNormalTexture.Scale);
+		}
+		else if (JsonMaterial.ShadingModel == EGLTFJsonShadingModel::Sheen)
+		{
+			GetValue(Inputs::SheenColorFactor, JsonMaterial.Sheen.ColorFactor);
+			GetValue(Inputs::SheenColorTexture, JsonMaterial.Sheen.ColorTexture);
+			GetValue(Inputs::SheenRoughnessFactor, JsonMaterial.Sheen.RoughnessFactor);
+			GetValue(Inputs::SheenRoughnessTexture, JsonMaterial.Sheen.RoughnessTexture);
+		}
+		else if (JsonMaterial.ShadingModel == EGLTFJsonShadingModel::Transmission)
+		{
+			GetValue(Inputs::TransmissionFactor, JsonMaterial.Transmission.Factor);
+			GetValue(Inputs::TransmissionTexture, JsonMaterial.Transmission.Texture);
+
+			float AlphaMode;
+			if (GetValue(Inputs::AlphaMode, AlphaMode))
+			{
+				JsonMaterial.AlphaMode = EGLTFJsonAlphaMode(FMath::RoundToInt(AlphaMode));
+			}
+		}
+		else if (JsonMaterial.ShadingModel == EGLTFJsonShadingModel::SpecularGlossiness)
+		{
+			GetValue(Inputs::SpecFactor, JsonMaterial.PBRSpecularGlossiness.SpecularFactor);
+			GetValue(Inputs::GlossinessFactor, JsonMaterial.PBRSpecularGlossiness.GlossinessFactor);
+			GetValue(Inputs::SpecularGlossinessTexture, JsonMaterial.PBRSpecularGlossiness.SpecularGlossinessTexture);
+		}
+	}
+}
\ No newline at end of file
diff --git a/Engine/Plugins/Enterprise/GLTFExporter/Source/GLTFExporter/Private/Converters/GLTFMaterialUtilities.h b/Engine/Plugins/Enterprise/GLTFExporter/Source/GLTFExporter/Private/Converters/GLTFMaterialUtilities.h
--- a/Engine/Plugins/Enterprise/GLTFExporter/Source/GLTFExporter/Private/Converters/GLTFMaterialUtilities.h
+++ b/Engine/Plugins/Enterprise/GLTFExporter/Source/GLTFExporter/Private/Converters/GLTFMaterialUtilities.h
@@ -94,3 +94,26 @@ private:
 	static void GetAllInputExpressionsOfType(const UMaterialInterface* Material, const FMaterialPropertyEx& Property, TArray<ExpressionType*>& OutExpressions);
 #endif
 };
+
+//Helper used to identify Interchange - glTF Importer created materials.
+// And then helps to acquire the Values set by the glTF Importer
+struct FGLTFImportMaterialMatchMakingHelper
+{
+	TMap<FString, UMaterialExpression*> Inputs;
+	FGLTFConvertBuilder& Builder;
+	FGLTFJsonMaterial& JsonMaterial;
+	const UMaterialInterface* Material;
+	bool bMaterialInstance; //uses Parameter acquisition approach if true, otherwise uses Inputs
+	bool bIsGLTFImportedMaterial;
+
+	FGLTFImportMaterialMatchMakingHelper(FGLTFConvertBuilder& InBuilder,
+		const UMaterialInterface* InMaterial,
+		FGLTFJsonMaterial& InJsonMaterial);
+
+	bool GetValue(const FString& InputKey, float& OutValue);
+	bool GetValue(const FString& InputKey, FGLTFJsonColor3& OutValue);
+	bool GetValue(const FString& InputKey, FGLTFJsonColor4& OutValue, bool HandleAsColor);
+	bool GetValue(const FString& InputKey, FGLTFJsonTextureInfo& OutValue);
+
+	void Process();
+};
\ No newline at end of file
diff --git a/Engine/Plugins/Enterprise/GLTFExporter/Source/GLTFExporter/Private/Converters/GLTFTextureConverters.cpp b/Engine/Plugins/Enterprise/GLTFExporter/Source/GLTFExporter/Private/Converters/GLTFTextureConverters.cpp
--- a/Engine/Plugins/Enterprise/GLTFExporter/Source/GLTFExporter/Private/Converters/GLTFTextureConverters.cpp
+++ b/Engine/Plugins/Enterprise/GLTFExporter/Source/GLTFExporter/Private/Converters/GLTFTextureConverters.cpp
@@ -3,7 +3,7 @@
 #include "Converters/GLTFTextureConverters.h"
 #include "Tasks/GLTFDelayedTextureTasks.h"
 
-FGLTFJsonTexture* FGLTFTexture2DConverter::Convert(const UTexture2D* Texture2D, bool bToSRGB)
+FGLTFJsonTexture* FGLTFTexture2DConverter::Convert(const UTexture2D* Texture2D, bool bToSRGB, TextureAddress WrapS, TextureAddress WrapT)
 {
 	if (Builder.ExportOptions->TextureImageFormat == EGLTFTextureImageFormat::None)
 	{
@@ -11,7 +11,7 @@ FGLTFJsonTexture* FGLTFTexture2DConverter::Convert(const UTexture2D* Texture2D,
 	}
 
 	FGLTFJsonTexture* JsonTexture = Builder.AddTexture();
-	Builder.ScheduleSlowTask<FGLTFDelayedTexture2DTask>(Builder, Texture2D, bToSRGB, JsonTexture);
+	Builder.ScheduleSlowTask<FGLTFDelayedTexture2DTask>(Builder, Texture2D, bToSRGB, JsonTexture, WrapS, WrapT);
 	return JsonTexture;
 }
 
diff --git a/Engine/Plugins/Enterprise/GLTFExporter/Source/GLTFExporter/Private/Json/GLTFJsonMaterial.cpp b/Engine/Plugins/Enterprise/GLTFExporter/Source/GLTFExporter/Private/Json/GLTFJsonMaterial.cpp
--- a/Engine/Plugins/Enterprise/GLTFExporter/Source/GLTFExporter/Private/Json/GLTFJsonMaterial.cpp
+++ b/Engine/Plugins/Enterprise/GLTFExporter/Source/GLTFExporter/Private/Json/GLTFJsonMaterial.cpp
@@ -177,6 +177,34 @@ void FGLTFJsonTransmissionExtension::WriteObject(IGLTFJsonWriter& Writer) const
 	}
 }
 
+void FGLTFJsonSpecularGlossinessExtension::WriteObject(IGLTFJsonWriter& Writer) const
+{
+	if (!DiffuseFactor.IsNearlyEqual(FGLTFJsonColor4::White))
+	{
+		Writer.Write(TEXT("diffuseFactor"), DiffuseFactor);
+	}
+
+	if (DiffuseTexture.Index != nullptr)
+	{
+		Writer.Write(TEXT("diffuseTexture"), DiffuseTexture);
+	}
+
+	if (!SpecularFactor.IsNearlyEqual(FGLTFJsonColor3::White))
+	{
+		Writer.Write(TEXT("specularFactor"), SpecularFactor);
+	}
+
+	if (!FMath::IsNearlyEqual(GlossinessFactor, 1.0f))
+	{
+		Writer.Write(TEXT("glossinessFactor"), GlossinessFactor);
+	}
+
+	if (SpecularGlossinessTexture.Index != nullptr)
+	{
+		Writer.Write(TEXT("specularGlossinessTexture"), SpecularGlossinessTexture);
+	}
+}
+
 void FGLTFJsonMaterial::WriteObject(IGLTFJsonWriter& Writer) const
 {
 	if (!Name.IsEmpty())
@@ -184,7 +212,8 @@ void FGLTFJsonMaterial::WriteObject(IGLTFJsonWriter& Writer) const
 		Writer.Write(TEXT("name"), Name);
 	}
 
-	if (ShadingModel != EGLTFJsonShadingModel::None)
+	if (ShadingModel != EGLTFJsonShadingModel::None &&
+		ShadingModel != EGLTFJsonShadingModel::SpecularGlossiness)
 	{
 		Writer.Write(TEXT("pbrMetallicRoughness"), PBRMetallicRoughness);
 	}
@@ -228,11 +257,11 @@ void FGLTFJsonMaterial::WriteObject(IGLTFJsonWriter& Writer) const
 	if (ShadingModel == EGLTFJsonShadingModel::Unlit ||
 		ShadingModel == EGLTFJsonShadingModel::ClearCoat ||
 		(ShadingModel == EGLTFJsonShadingModel::Transmission && Transmission.HasValue()) ||
+		(ShadingModel == EGLTFJsonShadingModel::SpecularGlossiness && PBRSpecularGlossiness.HasValue()) ||
 		HasEmissiveStrength || 
 		Specular.HasValue() || 
 		IOR.HasValue() ||
-		Sheen.HasValue() ||
-		Transmission.HasValue())
+		Sheen.HasValue())
 	{
 		Writer.StartExtensions();
 
@@ -261,10 +290,14 @@ void FGLTFJsonMaterial::WriteObject(IGLTFJsonWriter& Writer) const
 		{
 			Writer.Write(EGLTFJsonExtension::KHR_MaterialsClearCoat, ClearCoat);
 		}
-		else if (ShadingModel == EGLTFJsonShadingModel::Transmission)
+		else if (ShadingModel == EGLTFJsonShadingModel::Transmission && Transmission.HasValue())
 		{
 			Writer.Write(EGLTFJsonExtension::KHR_MaterialsTransmission, Transmission);
 		}
+		else if (ShadingModel == EGLTFJsonShadingModel::SpecularGlossiness && PBRSpecularGlossiness.HasValue())
+		{
+			Writer.Write(EGLTFJsonExtension::KHR_MaterialsSpecularGlossiness, PBRSpecularGlossiness);
+		}
 
 		if (HasEmissiveStrength)
 		{
diff --git a/Engine/Plugins/Enterprise/GLTFExporter/Source/GLTFExporter/Private/Options/GLTFExportOptions.cpp b/Engine/Plugins/Enterprise/GLTFExporter/Source/GLTFExporter/Private/Options/GLTFExportOptions.cpp
--- a/Engine/Plugins/Enterprise/GLTFExporter/Source/GLTFExporter/Private/Options/GLTFExportOptions.cpp
+++ b/Engine/Plugins/Enterprise/GLTFExporter/Source/GLTFExporter/Private/Options/GLTFExportOptions.cpp
@@ -15,11 +15,13 @@ void UGLTFExportOptions::ResetToDefault()
 	bSkipNearDefaultValues = true;
 	bIncludeCopyrightNotice = false;
 	bExportProxyMaterials = true;
+	bUseImporterMaterialMapping = true;
 	bExportUnlitMaterials = true;
 	bExportClearCoatMaterials = true;
 	bExportClothMaterials = true;
 	bExportThinTranslucentMaterials = false;
 	bExportEmissiveStrength = true;
+	bExportSpecularGlossinessMaterials = true;
 	BakeMaterialInputs = EGLTFMaterialBakeMode::UseMeshData;
 	DefaultMaterialBakeSize = FGLTFMaterialBakeSize::Default;
 	DefaultMaterialBakeFilter = TF_Trilinear;
diff --git a/Engine/Plugins/Enterprise/GLTFExporter/Source/GLTFExporter/Private/Tasks/GLTFDelayedMaterialTasks.cpp b/Engine/Plugins/Enterprise/GLTFExporter/Source/GLTFExporter/Private/Tasks/GLTFDelayedMaterialTasks.cpp
--- a/Engine/Plugins/Enterprise/GLTFExporter/Source/GLTFExporter/Private/Tasks/GLTFDelayedMaterialTasks.cpp
+++ b/Engine/Plugins/Enterprise/GLTFExporter/Source/GLTFExporter/Private/Tasks/GLTFDelayedMaterialTasks.cpp
@@ -57,12 +57,26 @@ void FGLTFDelayedMaterialTask::Process()
 	}
 
 	JsonMaterial->Name = GetMaterialName();
-	JsonMaterial->AlphaCutoff = Material->GetOpacityMaskClipValue();
 	JsonMaterial->DoubleSided = Material->IsTwoSided();
-
-	ConvertShadingModel(JsonMaterial->ShadingModel);
 	ConvertAlphaMode(JsonMaterial->AlphaMode);
 
+	FString WarningMessages;
+	EMaterialShadingModel UEShadingModel = FGLTFMaterialUtilities::GetShadingModel(Material, WarningMessages);
+	if (!WarningMessages.IsEmpty())
+	{
+		Builder.LogWarning(WarningMessages);
+	}
+
+	if (HandleGLTFImported(UEShadingModel))
+	{
+		return;
+	}
+
+	JsonMaterial->AlphaCutoff = Material->GetOpacityMaskClipValue();
+
+	ConvertShadingModel(UEShadingModel, JsonMaterial->ShadingModel);
+	ApplyExportOptionsToShadingModel(JsonMaterial->ShadingModel, UEShadingModel);
+
 	if (!ensure((JsonMaterial->ShadingModel != EGLTFJsonShadingModel::None) && ((int)JsonMaterial->ShadingModel < (int)EGLTFJsonShadingModel::NumShadingModels)))
 	{
 		Builder.LogWarning(FString::Printf(TEXT("Failed to export %s material. ShadingModel is unexpected."), *Material->GetName()));
@@ -406,79 +420,101 @@ void FGLTFDelayedMaterialTask::GetProxyParameter(const FGLTFProxyMaterialTexture
 	}
 }
 
-void FGLTFDelayedMaterialTask::ConvertShadingModel(EGLTFJsonShadingModel& OutShadingModel) const
+void FGLTFDelayedMaterialTask::ConvertShadingModel(EMaterialShadingModel& UEShadingModel, EGLTFJsonShadingModel& OutGLTFShadingModel) const
 {
-	FString WarningMessage;
-	EMaterialShadingModel ShadingModel = FGLTFMaterialUtilities::GetShadingModel(Material, WarningMessage);
-
-	if (!WarningMessage.IsEmpty())
-	{
-		Builder.LogWarning(WarningMessage);
-	}
-
 	const EBlendMode BlendMode = Material->GetBlendMode();
-	if (ShadingModel == MSM_ClearCoat && BlendMode != BLEND_Opaque && BlendMode != BLEND_Masked)
+	if (UEShadingModel == MSM_ClearCoat && BlendMode != BLEND_Opaque && BlendMode != BLEND_Masked)
 	{
 		// NOTE: Unreal seems to disable clear coat when blend mode anything but opaque or masked
-		ShadingModel = MSM_DefaultLit;
+		UEShadingModel = MSM_DefaultLit;
 	}
 
-	OutShadingModel = FGLTFCoreUtilities::ConvertShadingModel(ShadingModel);
-	if (OutShadingModel == EGLTFJsonShadingModel::None)
+	OutGLTFShadingModel = FGLTFCoreUtilities::ConvertShadingModel(UEShadingModel);
+}
+
+void FGLTFDelayedMaterialTask::ApplyExportOptionsToShadingModel(EGLTFJsonShadingModel& ShadingModel, const EMaterialShadingModel& UEMaterialShadingModel) const
+{
+	switch (ShadingModel)
 	{
-		OutShadingModel = EGLTFJsonShadingModel::Default;
+		case EGLTFJsonShadingModel::Default:
+			break;
 
-		Builder.LogWarning(FString::Printf(
-			TEXT("Unsupported shading model (%s) in material %s, will export as %s"),
-			*FGLTFNameUtilities::GetName(ShadingModel),
-			*Material->GetName(),
-			*FGLTFNameUtilities::GetName(MSM_DefaultLit)));
-		return;
-	}
+		case EGLTFJsonShadingModel::Unlit:
+			if (!Builder.ExportOptions->bExportUnlitMaterials)
+			{
+				Builder.LogWarning(FString::Printf(
+					TEXT("Shading model (%s) in material %s disabled by export options, will export as %s"),
+					*FGLTFNameUtilities::GetName(UEMaterialShadingModel),
+					*Material->GetName(),
+					*FGLTFNameUtilities::GetName(MSM_DefaultLit)));
 
-	if (OutShadingModel == EGLTFJsonShadingModel::Unlit && !Builder.ExportOptions->bExportUnlitMaterials)
-	{
-		OutShadingModel = EGLTFJsonShadingModel::Default;
+				ShadingModel = EGLTFJsonShadingModel::Default;
+			}
+			break;
 
-		Builder.LogWarning(FString::Printf(
-			TEXT("Shading model (%s) in material %s disabled by export options, will export as %s"),
-			*FGLTFNameUtilities::GetName(ShadingModel),
-			*Material->GetName(),
-			*FGLTFNameUtilities::GetName(MSM_DefaultLit)));
-		return;
-	}
+		case EGLTFJsonShadingModel::ClearCoat:
+			if (!Builder.ExportOptions->bExportClearCoatMaterials)
+			{
+				Builder.LogWarning(FString::Printf(
+					TEXT("Shading model (%s) in material %s disabled by export options, will export as %s"),
+					*FGLTFNameUtilities::GetName(UEMaterialShadingModel),
+					*Material->GetName(),
+					*FGLTFNameUtilities::GetName(MSM_DefaultLit)));
 
-	if (OutShadingModel == EGLTFJsonShadingModel::ClearCoat && !Builder.ExportOptions->bExportClearCoatMaterials)
-	{
-		OutShadingModel = EGLTFJsonShadingModel::Default;
+				ShadingModel = EGLTFJsonShadingModel::Default;
+			}
+			break;
 
-		Builder.LogWarning(FString::Printf(
-			TEXT("Shading model (%s) in material %s disabled by export options, will export as %s"),
-			*FGLTFNameUtilities::GetName(ShadingModel),
-			*Material->GetName(),
-			*FGLTFNameUtilities::GetName(MSM_DefaultLit)));
-	}
+		case EGLTFJsonShadingModel::Sheen:
+			if (ShadingModel == EGLTFJsonShadingModel::Sheen && !Builder.ExportOptions->bExportClothMaterials)
+			{
+				Builder.LogWarning(FString::Printf(
+					TEXT("Shading model (%s) in material %s disabled by export options, will export as %s"),
+					*FGLTFNameUtilities::GetName(UEMaterialShadingModel),
+					*Material->GetName(),
+					*FGLTFNameUtilities::GetName(MSM_DefaultLit)));
 
-	if (OutShadingModel == EGLTFJsonShadingModel::Sheen && !Builder.ExportOptions->bExportClothMaterials)
-	{
-		OutShadingModel = EGLTFJsonShadingModel::Default;
+				ShadingModel = EGLTFJsonShadingModel::Default;
+			}
+			break;
 
-		Builder.LogWarning(FString::Printf(
-			TEXT("Shading model (%s) in material %s disabled by export options, will export as %s"),
-			*FGLTFNameUtilities::GetName(ShadingModel),
-			*Material->GetName(),
-			*FGLTFNameUtilities::GetName(MSM_DefaultLit)));
-	}
+		case EGLTFJsonShadingModel::Transmission:
+			if (!Builder.ExportOptions->bExportThinTranslucentMaterials)
+			{
+				Builder.LogWarning(FString::Printf(
+					TEXT("Shading model (%s) in material %s disabled by export options, will export as %s"),
+					*FGLTFNameUtilities::GetName(UEMaterialShadingModel),
+					*Material->GetName(),
+					*FGLTFNameUtilities::GetName(MSM_DefaultLit)));
 
-	if (OutShadingModel == EGLTFJsonShadingModel::Transmission && !Builder.ExportOptions->bExportThinTranslucentMaterials)
-	{
-		OutShadingModel = EGLTFJsonShadingModel::Default;
+				ShadingModel = EGLTFJsonShadingModel::Default;
+			}
+			break;
 
-		Builder.LogWarning(FString::Printf(
-			TEXT("Shading model (%s) in material %s disabled by export options, will export as %s"),
-			*FGLTFNameUtilities::GetName(ShadingModel),
-			*Material->GetName(),
-			*FGLTFNameUtilities::GetName(MSM_DefaultLit)));
+		case EGLTFJsonShadingModel::SpecularGlossiness:
+			if (!Builder.ExportOptions->bExportSpecularGlossinessMaterials)
+			{
+				Builder.LogWarning(FString::Printf(
+					TEXT("Shading model (%s) in material %s disabled by export options, will export as %s"),
+					*FGLTFNameUtilities::GetName(UEMaterialShadingModel),
+					*Material->GetName(),
+					*FGLTFNameUtilities::GetName(MSM_DefaultLit)));
+
+				ShadingModel = EGLTFJsonShadingModel::Default;
+			}
+			break;
+
+		default:
+			{
+				Builder.LogWarning(FString::Printf(
+					TEXT("Unsupported shading model (%s) in material %s, will export as %s"),
+					*FGLTFNameUtilities::GetName(UEMaterialShadingModel),
+					*Material->GetName(),
+					*FGLTFNameUtilities::GetName(MSM_DefaultLit)));
+
+				ShadingModel = EGLTFJsonShadingModel::Default;
+			}
+			break;
 	}
 }
 
@@ -1995,3 +2031,25 @@ void FGLTFDelayedMaterialTask::CombinePixels(const TArray<FColor>& FirstPixels,
 }
 
 #endif
+
+bool FGLTFDelayedMaterialTask::HandleGLTFImported(const EMaterialShadingModel& UEMaterialShadingModel)
+{
+	if (!Builder.ExportOptions->bExportUnlitMaterials)
+	{
+		return false;
+	}
+
+	FGLTFImportMaterialMatchMakingHelper GLTFImportedProcesser(Builder, Material, *JsonMaterial);
+
+	if (!GLTFImportedProcesser.bIsGLTFImportedMaterial)
+	{
+		return false;
+	}
+
+	ApplyExportOptionsToShadingModel(JsonMaterial->ShadingModel, UEMaterialShadingModel);
+
+	GLTFImportedProcesser.Process();
+
+	return true;
+}
+
diff --git a/Engine/Plugins/Enterprise/GLTFExporter/Source/GLTFExporter/Private/Tasks/GLTFDelayedMaterialTasks.h b/Engine/Plugins/Enterprise/GLTFExporter/Source/GLTFExporter/Private/Tasks/GLTFDelayedMaterialTasks.h
--- a/Engine/Plugins/Enterprise/GLTFExporter/Source/GLTFExporter/Private/Tasks/GLTFDelayedMaterialTasks.h
+++ b/Engine/Plugins/Enterprise/GLTFExporter/Source/GLTFExporter/Private/Tasks/GLTFDelayedMaterialTasks.h
@@ -64,7 +64,9 @@ private:
 		return ParameterInfo.Get(Material, Value, true);
 	}
 
-	void ConvertShadingModel(EGLTFJsonShadingModel& OutShadingModel) const;
+	void ConvertShadingModel(EMaterialShadingModel& UEShadingModel, EGLTFJsonShadingModel& OutGLTFShadingModel) const;
+	void ApplyExportOptionsToShadingModel(EGLTFJsonShadingModel& ShadingModel, const EMaterialShadingModel& UEMaterialShadingModel) const;
+
 	void ConvertAlphaMode(EGLTFJsonAlphaMode& OutAlphaMode) const;
 
 #if WITH_EDITOR
@@ -108,6 +110,7 @@ private:
 
 	template <typename CallbackType>
 	static void CombinePixels(const TArray<FColor>& FirstPixels, const TArray<FColor>& SecondPixels, TArray<FColor>& OutPixels, CallbackType Callback);
-
 #endif
+
+	bool HandleGLTFImported(const EMaterialShadingModel& UEMaterialShadingModel);
 };
diff --git a/Engine/Plugins/Enterprise/GLTFExporter/Source/GLTFExporter/Private/Tasks/GLTFDelayedTextureTasks.cpp b/Engine/Plugins/Enterprise/GLTFExporter/Source/GLTFExporter/Private/Tasks/GLTFDelayedTextureTasks.cpp
--- a/Engine/Plugins/Enterprise/GLTFExporter/Source/GLTFExporter/Private/Tasks/GLTFDelayedTextureTasks.cpp
+++ b/Engine/Plugins/Enterprise/GLTFExporter/Source/GLTFExporter/Private/Tasks/GLTFDelayedTextureTasks.cpp
@@ -47,7 +47,11 @@ void FGLTFDelayedTexture2DTask::Process()
 	const bool bIgnoreAlpha = FGLTFTextureUtilities::IsAlphaless(Texture2D->GetPixelFormat());
 
 	JsonTexture->Source = Builder.AddUniqueImage(Pixels, Size, bIgnoreAlpha, JsonTexture->Name);
-	JsonTexture->Sampler = Builder.AddUniqueSampler(Texture2D);
+
+	if (TextureAddressX == TextureAddress::TA_MAX) TextureAddressX = Texture2D->GetTextureAddressX();
+	if (TextureAddressY == TextureAddress::TA_MAX) TextureAddressY = Texture2D->GetTextureAddressX();
+
+	JsonTexture->Sampler = Builder.AddUniqueSampler(TextureAddressX, TextureAddressY, Texture2D->Filter, Texture2D->LODGroup);
 }
 
 FString FGLTFDelayedTextureRenderTarget2DTask::GetName()
diff --git a/Engine/Plugins/Enterprise/GLTFExporter/Source/GLTFExporter/Private/Tasks/GLTFDelayedTextureTasks.h b/Engine/Plugins/Enterprise/GLTFExporter/Source/GLTFExporter/Private/Tasks/GLTFDelayedTextureTasks.h
--- a/Engine/Plugins/Enterprise/GLTFExporter/Source/GLTFExporter/Private/Tasks/GLTFDelayedTextureTasks.h
+++ b/Engine/Plugins/Enterprise/GLTFExporter/Source/GLTFExporter/Private/Tasks/GLTFDelayedTextureTasks.h
@@ -9,11 +9,13 @@ class FGLTFDelayedTexture2DTask : public FGLTFDelayedTask
 {
 public:
 
-	FGLTFDelayedTexture2DTask(FGLTFConvertBuilder& Builder, const UTexture2D* Texture2D, bool bToSRGB, FGLTFJsonTexture* JsonTexture)
+	FGLTFDelayedTexture2DTask(FGLTFConvertBuilder& Builder, const UTexture2D* Texture2D, bool bToSRGB, FGLTFJsonTexture* JsonTexture, TextureAddress InTextureAddressX, TextureAddress InTextureAddressY)
 		: FGLTFDelayedTask(EGLTFTaskPriority::Texture)
 		, Builder(Builder)
 		, Texture2D(Texture2D)
 		, bToSRGB(bToSRGB)
+		, TextureAddressX(InTextureAddressX)
+		, TextureAddressY(InTextureAddressY)
 		, JsonTexture(JsonTexture)
 	{
 	}
@@ -27,6 +29,8 @@ private:
 	FGLTFConvertBuilder& Builder;
 	const UTexture2D* Texture2D;
 	bool bToSRGB;
+	TextureAddress TextureAddressX;
+	TextureAddress TextureAddressY;
 	FGLTFJsonTexture* JsonTexture;
 };
 
diff --git a/Engine/Plugins/Enterprise/GLTFExporter/Source/GLTFExporter/Private/Utilities/GLTFJsonUtilities.cpp b/Engine/Plugins/Enterprise/GLTFExporter/Source/GLTFExporter/Private/Utilities/GLTFJsonUtilities.cpp
--- a/Engine/Plugins/Enterprise/GLTFExporter/Source/GLTFExporter/Private/Utilities/GLTFJsonUtilities.cpp
+++ b/Engine/Plugins/Enterprise/GLTFExporter/Source/GLTFExporter/Private/Utilities/GLTFJsonUtilities.cpp
@@ -6,17 +6,18 @@ const TCHAR* FGLTFJsonUtilities::GetValue(EGLTFJsonExtension Enum)
 {
 	switch (Enum)
 	{
-		case EGLTFJsonExtension::KHR_LightsPunctual:            return TEXT("KHR_lights_punctual");
-		case EGLTFJsonExtension::KHR_MaterialsClearCoat:        return TEXT("KHR_materials_clearcoat");
-		case EGLTFJsonExtension::KHR_MaterialsEmissiveStrength: return TEXT("KHR_materials_emissive_strength");
-		case EGLTFJsonExtension::KHR_MaterialsUnlit:            return TEXT("KHR_materials_unlit");
-		case EGLTFJsonExtension::KHR_MaterialsVariants:         return TEXT("KHR_materials_variants");
-		case EGLTFJsonExtension::KHR_MaterialsIOR:              return TEXT("KHR_materials_ior");
-		case EGLTFJsonExtension::KHR_MaterialsSheen:            return TEXT("KHR_materials_sheen");
-		case EGLTFJsonExtension::KHR_MaterialsTransmission:     return TEXT("KHR_materials_transmission");
-		case EGLTFJsonExtension::KHR_MeshQuantization:          return TEXT("KHR_mesh_quantization");
-		case EGLTFJsonExtension::KHR_TextureTransform:          return TEXT("KHR_texture_transform");
-		case EGLTFJsonExtension::KHR_MaterialsSpecular:	        return TEXT("KHR_materials_specular");
+		case EGLTFJsonExtension::KHR_LightsPunctual:                  return TEXT("KHR_lights_punctual");
+		case EGLTFJsonExtension::KHR_MaterialsClearCoat:              return TEXT("KHR_materials_clearcoat");
+		case EGLTFJsonExtension::KHR_MaterialsEmissiveStrength:       return TEXT("KHR_materials_emissive_strength");
+		case EGLTFJsonExtension::KHR_MaterialsUnlit:                  return TEXT("KHR_materials_unlit");
+		case EGLTFJsonExtension::KHR_MaterialsVariants:               return TEXT("KHR_materials_variants");
+		case EGLTFJsonExtension::KHR_MaterialsIOR:                    return TEXT("KHR_materials_ior");
+		case EGLTFJsonExtension::KHR_MaterialsSheen:                  return TEXT("KHR_materials_sheen");
+		case EGLTFJsonExtension::KHR_MaterialsTransmission:           return TEXT("KHR_materials_transmission");
+		case EGLTFJsonExtension::KHR_MaterialsSpecularGlossiness:     return TEXT("KHR_materials_pbrSpecularGlossiness");
+		case EGLTFJsonExtension::KHR_MeshQuantization:                return TEXT("KHR_mesh_quantization");
+		case EGLTFJsonExtension::KHR_TextureTransform:                return TEXT("KHR_texture_transform");
+		case EGLTFJsonExtension::KHR_MaterialsSpecular:	              return TEXT("KHR_materials_specular");
 		default:
 			checkNoEntry();
 			return TEXT("");
@@ -121,11 +122,12 @@ const TCHAR* FGLTFJsonUtilities::GetValue(EGLTFJsonShadingModel Enum)
 {
 	switch (Enum)
 	{
-		case EGLTFJsonShadingModel::Default:          return TEXT("Default");
-		case EGLTFJsonShadingModel::Unlit:            return TEXT("Unlit");
-		case EGLTFJsonShadingModel::ClearCoat:        return TEXT("ClearCoat");
-		case EGLTFJsonShadingModel::Sheen:            return TEXT("Sheen");
-		case EGLTFJsonShadingModel::Transmission:     return TEXT("Transmission");
+		case EGLTFJsonShadingModel::Default:                return TEXT("Default");
+		case EGLTFJsonShadingModel::Unlit:                  return TEXT("Unlit");
+		case EGLTFJsonShadingModel::ClearCoat:              return TEXT("ClearCoat");
+		case EGLTFJsonShadingModel::Sheen:                  return TEXT("Sheen");
+		case EGLTFJsonShadingModel::Transmission:           return TEXT("Transmission");
+		case EGLTFJsonShadingModel::SpecularGlossiness:     return TEXT("SpecularGlossiness");
 		default:
 			checkNoEntry();
 			return TEXT("");
diff --git a/Engine/Plugins/Enterprise/GLTFExporter/Source/GLTFExporter/Public/Builders/GLTFConvertBuilder.h b/Engine/Plugins/Enterprise/GLTFExporter/Source/GLTFExporter/Public/Builders/GLTFConvertBuilder.h
--- a/Engine/Plugins/Enterprise/GLTFExporter/Source/GLTFExporter/Public/Builders/GLTFConvertBuilder.h
+++ b/Engine/Plugins/Enterprise/GLTFExporter/Source/GLTFExporter/Public/Builders/GLTFConvertBuilder.h
@@ -64,8 +64,8 @@ public:
 	FGLTFJsonTexture* AddUniqueTexture(const UTexture* Texture);
 	FGLTFJsonTexture* AddUniqueTexture(const UTexture2D* Texture);
 	FGLTFJsonTexture* AddUniqueTexture(const UTextureRenderTarget2D* Texture);
-	FGLTFJsonTexture* AddUniqueTexture(const UTexture* Texture, bool bToSRGB);
-	FGLTFJsonTexture* AddUniqueTexture(const UTexture2D* Texture, bool bToSRGB);
+	FGLTFJsonTexture* AddUniqueTexture(const UTexture* Texture, bool bToSRGB, TextureAddress TextureAddressX = TextureAddress::TA_MAX, TextureAddress TextureAddressY = TextureAddress::TA_MAX);
+	FGLTFJsonTexture* AddUniqueTexture(const UTexture2D* Texture, bool bToSRGB, TextureAddress TextureAddressX, TextureAddress TextureAddressY);
 	FGLTFJsonTexture* AddUniqueTexture(const UTextureRenderTarget2D* Texture, bool bToSRGB);
 	FGLTFJsonImage* AddUniqueImage(TGLTFSharedArray<FColor>& Pixels, FIntPoint Size, bool bIgnoreAlpha, const FString& Name);
 
diff --git a/Engine/Plugins/Enterprise/GLTFExporter/Source/GLTFExporter/Public/Converters/GLTFTextureConverters.h b/Engine/Plugins/Enterprise/GLTFExporter/Source/GLTFExporter/Public/Converters/GLTFTextureConverters.h
--- a/Engine/Plugins/Enterprise/GLTFExporter/Source/GLTFExporter/Public/Converters/GLTFTextureConverters.h
+++ b/Engine/Plugins/Enterprise/GLTFExporter/Source/GLTFExporter/Public/Converters/GLTFTextureConverters.h
@@ -6,12 +6,14 @@
 #include "Converters/GLTFConverter.h"
 #include "Converters/GLTFBuilderContext.h"
 
+#include "Engine/TextureDefines.h"
+
 class UTexture2D;
 class UTextureRenderTarget2D;
 
 // TODO: generalize parameter bToSRGB to SamplerType (including normalmap unpacking)
 
-typedef TGLTFConverter<FGLTFJsonTexture*, const UTexture2D*, bool> IGLTFTexture2DConverter;
+typedef TGLTFConverter<FGLTFJsonTexture*, const UTexture2D*, bool, TextureAddress, TextureAddress> IGLTFTexture2DConverter;
 typedef TGLTFConverter<FGLTFJsonTexture*, const UTextureRenderTarget2D*, bool> IGLTFTextureRenderTarget2DConverter;
 
 class GLTFEXPORTER_API FGLTFTexture2DConverter : public FGLTFBuilderContext, public IGLTFTexture2DConverter
@@ -22,7 +24,7 @@ public:
 
 protected:
 
-	virtual FGLTFJsonTexture* Convert(const UTexture2D* Texture2D, bool bToSRGB) override;
+	virtual FGLTFJsonTexture* Convert(const UTexture2D* Texture2D, bool bToSRGB, TextureAddress WrapS, TextureAddress WrapT) override;
 };
 
 class GLTFEXPORTER_API FGLTFTextureRenderTarget2DConverter : public FGLTFBuilderContext, public IGLTFTextureRenderTarget2DConverter
diff --git a/Engine/Plugins/Enterprise/GLTFExporter/Source/GLTFExporter/Public/Json/GLTFJsonEnums.h b/Engine/Plugins/Enterprise/GLTFExporter/Source/GLTFExporter/Public/Json/GLTFJsonEnums.h
--- a/Engine/Plugins/Enterprise/GLTFExporter/Source/GLTFExporter/Public/Json/GLTFJsonEnums.h
+++ b/Engine/Plugins/Enterprise/GLTFExporter/Source/GLTFExporter/Public/Json/GLTFJsonEnums.h
@@ -12,6 +12,7 @@ enum class EGLTFJsonExtension
 	KHR_MaterialsIOR,
 	KHR_MaterialsSheen,
 	KHR_MaterialsTransmission,
+	KHR_MaterialsSpecularGlossiness,
 	KHR_MeshQuantization,
 	KHR_TextureTransform,
 	KHR_MaterialsSpecular
@@ -25,6 +26,7 @@ enum class EGLTFJsonShadingModel
 	ClearCoat,
 	Sheen, /*UE Cloth*/
 	Transmission, /*UE Thin Translucent*/
+	SpecularGlossiness, /*Only for Importer material mapping.*/
 
 	NumShadingModels
 };
diff --git a/Engine/Plugins/Enterprise/GLTFExporter/Source/GLTFExporter/Public/Json/GLTFJsonMaterial.h b/Engine/Plugins/Enterprise/GLTFExporter/Source/GLTFExporter/Public/Json/GLTFJsonMaterial.h
--- a/Engine/Plugins/Enterprise/GLTFExporter/Source/GLTFExporter/Public/Json/GLTFJsonMaterial.h
+++ b/Engine/Plugins/Enterprise/GLTFExporter/Source/GLTFExporter/Public/Json/GLTFJsonMaterial.h
@@ -150,6 +150,27 @@ struct GLTFEXPORTER_API FGLTFJsonTransmissionExtension : IGLTFJsonObject
 	virtual void WriteObject(IGLTFJsonWriter& Writer) const override;
 };
 
+struct GLTFEXPORTER_API FGLTFJsonSpecularGlossinessExtension : IGLTFJsonObject
+{
+	FGLTFJsonColor4         DiffuseFactor;
+	FGLTFJsonTextureInfo    DiffuseTexture;
+
+	FGLTFJsonColor3         SpecularFactor;
+	float                   GlossinessFactor;
+	FGLTFJsonTextureInfo    SpecularGlossinessTexture;
+
+	FGLTFJsonSpecularGlossinessExtension()
+		: DiffuseFactor(FGLTFJsonColor4::White)
+		, SpecularFactor(FGLTFJsonColor3::White)
+		, GlossinessFactor(1)
+	{
+	}
+
+	bool HasValue() const { return !DiffuseFactor.IsNearlyEqual(FGLTFJsonColor4::White) || DiffuseTexture.Index != nullptr || !SpecularFactor.IsNearlyEqual(FGLTFJsonColor3::White) || !FMath::IsNearlyEqual(GlossinessFactor, 1.0f) || SpecularGlossinessTexture.Index != nullptr; }
+
+	virtual void WriteObject(IGLTFJsonWriter& Writer) const override;
+};
+
 struct GLTFEXPORTER_API FGLTFJsonMaterial : IGLTFJsonIndexedObject
 {
 	FString Name;
@@ -157,6 +178,7 @@ struct GLTFEXPORTER_API FGLTFJsonMaterial : IGLTFJsonIndexedObject
 	EGLTFJsonShadingModel ShadingModel;
 
 	FGLTFJsonPBRMetallicRoughness PBRMetallicRoughness;
+	FGLTFJsonSpecularGlossinessExtension PBRSpecularGlossiness;
 
 	FGLTFJsonNormalTextureInfo NormalTexture;
 	FGLTFJsonOcclusionTextureInfo OcclusionTexture;
diff --git a/Engine/Plugins/Enterprise/GLTFExporter/Source/GLTFExporter/Public/Options/GLTFExportOptions.h b/Engine/Plugins/Enterprise/GLTFExporter/Source/GLTFExporter/Public/Options/GLTFExportOptions.h
--- a/Engine/Plugins/Enterprise/GLTFExporter/Source/GLTFExporter/Public/Options/GLTFExportOptions.h
+++ b/Engine/Plugins/Enterprise/GLTFExporter/Source/GLTFExporter/Public/Options/GLTFExportOptions.h
@@ -53,6 +53,10 @@ class GLTFEXPORTER_API UGLTFExportOptions : public UObject
 	UPROPERTY(EditAnywhere, BlueprintReadWrite, Config, Category = Material)
 	bool bExportProxyMaterials;
 
+	/** If enabled, materials imported with the Interchange-glTF importer will be directly mapped for the Exporter. bExport material options below will be ignored. */
+	UPROPERTY(EditAnywhere, BlueprintReadWrite, Config, Category = Material)
+	bool bUseImporterMaterialMapping;
+
 	/** If enabled, materials with shading model unlit will be properly exported. Uses extension KHR_materials_unlit. */
 	UPROPERTY(EditAnywhere, BlueprintReadWrite, Config, Category = Material)
 	bool bExportUnlitMaterials;
@@ -69,6 +73,10 @@ class GLTFEXPORTER_API UGLTFExportOptions : public UObject
 	UPROPERTY(EditAnywhere, BlueprintReadWrite, Config, Category = Material)
 	bool bExportThinTranslucentMaterials;
 
+	/** If enabled, materials using the Importer's SpecularGlossiness material function will be exported. Uses extension KHR_materials_pbrSpecularGlossiness. */
+	UPROPERTY(EditAnywhere, BlueprintReadWrite, Config, Category = Material)
+	bool bExportSpecularGlossinessMaterials;
+
 	/** If enabled, allows materials to have an emissive factor that exceeds the standard range [0.0, 1.0]. Uses extension KHR_materials_emissive_strength. */
 	UPROPERTY(EditAnywhere, BlueprintReadWrite, Config, Category = Material)
 	bool bExportEmissiveStrength;
diff --git a/Engine/Plugins/Interchange/Runtime/Source/Import/Private/Gltf/InterchangeGLTFMaterial.cpp b/Engine/Plugins/Interchange/Runtime/Source/Import/Private/Gltf/InterchangeGLTFMaterial.cpp
--- a/Engine/Plugins/Interchange/Runtime/Source/Import/Private/Gltf/InterchangeGLTFMaterial.cpp
+++ b/Engine/Plugins/Interchange/Runtime/Source/Import/Private/Gltf/InterchangeGLTFMaterial.cpp
@@ -77,24 +77,24 @@ namespace UE::Interchange::GLTFMaterials
 
 	FString GetMaterialFunctionPath(EShadingModel ShadingModel)
 	{
-		if (ShadingModelToMaterialFunctions.Contains(ShadingModel))
+		if (ShadingModelToMaterialInformation.Contains(ShadingModel))
 		{
-			return ShadingModelToMaterialFunctions[ShadingModel].Key;
+			return ShadingModelToMaterialInformation[ShadingModel].MaterialFunctionPath;
 		}
 
 		ensure(false);
-		return ShadingModelToMaterialFunctions.begin()->Value.Key;
+		return ShadingModelToMaterialInformation.begin()->Value.MaterialFunctionPath;
 	}
 
 	TArray<FString> GetOutputs(EShadingModel ShadingModel)
 	{
-		if (ShadingModelToMaterialFunctions.Contains(ShadingModel))
+		if (ShadingModelToMaterialInformation.Contains(ShadingModel))
 		{
-			return ShadingModelToMaterialFunctions[ShadingModel].Value;
+			return ShadingModelToMaterialInformation[ShadingModel].MaterialFunctionOutputs;
 		}
 
 		ensure(false);
-		return ShadingModelToMaterialFunctions.begin()->Value.Value;
+		return ShadingModelToMaterialInformation.begin()->Value.MaterialFunctionOutputs;
 	}
 
 	//Precedence order based on the logic in the InterchangeGenericMaterialPipeline: ClearCoat > Sheen > unlit
@@ -140,23 +140,19 @@ namespace UE::Interchange::GLTFMaterials
 
 	FVector4f GetTilingMethod(const GLTF::FSampler& Sampler)
 	{
-		float TilingU = 0;
-		switch (Sampler.WrapS)
-		{
-			case GLTF::FSampler::EWrap::Repeat:         TilingU = 0; break;
-			case GLTF::FSampler::EWrap::ClampToEdge:    TilingU = 1; break;
-			case GLTF::FSampler::EWrap::MirroredRepeat: TilingU = 2; break;
-			default:                                    TilingU = 0; break;;
-		}
+		auto GetSingleChannelTilingMethod = [](GLTF::FSampler::EWrap Wrap/*TilingMethod*/)
+			{
+				switch (Wrap)
+				{
+					case GLTF::FSampler::EWrap::Repeat:         return 0;
+					case GLTF::FSampler::EWrap::ClampToEdge:    return 1;
+					case GLTF::FSampler::EWrap::MirroredRepeat: return 2;
+					default:                                    return 0;
+				}
+			};
 
-		float TilingV = 0;
-		switch (Sampler.WrapT)
-		{
-			case GLTF::FSampler::EWrap::Repeat:         TilingV = 0; break;
-			case GLTF::FSampler::EWrap::ClampToEdge:    TilingV = 1; break;
-			case GLTF::FSampler::EWrap::MirroredRepeat: TilingV = 2; break;
-			default:                                    TilingV = 0; break;;
-		}
+		float TilingU = GetSingleChannelTilingMethod(Sampler.WrapS);
+		float TilingV = GetSingleChannelTilingMethod(Sampler.WrapT);
 
 		return FVector4f(TilingU, TilingV, 0, 0);
 	}
@@ -334,8 +330,8 @@ namespace UE::Interchange::GLTFMaterials
 		void SetColor(const FString& Name, const FVector4f& Value, const FVector4f& DefaultValue)
 		{
 			{//Material specific settings:
-				SetVec3(Name + Inputs::Color_RGB, FVector3f(Value), FVector3f(DefaultValue), EProcessType::MATERIAL);
-				SetScalar(Name + Inputs::Color_A, Value.W, DefaultValue.W, EProcessType::MATERIAL);
+				SetVec3(Name + Inputs::PostFix::Color_RGB, FVector3f(Value), FVector3f(DefaultValue), EProcessType::MATERIAL);
+				SetScalar(Name + Inputs::PostFix::Color_A, Value.W, DefaultValue.W, EProcessType::MATERIAL);
 			}
 
 			{//MaterialInstance specific settings:
@@ -343,7 +339,7 @@ namespace UE::Interchange::GLTFMaterials
 			}
 		}
 
-		void SetMap(const GLTF::FTextureMap& TextureMap, const FString& Name)
+		void SetMap(const FString& Name, const GLTF::FTextureMap& TextureMap)
 		{
 			if (Textures.IsValidIndex(TextureMap.TextureIndex))
 			{
@@ -368,25 +364,25 @@ namespace UE::Interchange::GLTFMaterials
 				// [2...3) -> UV2
 				// [3...4) -> UV3
 				// else    -> UV0 (defaults to 0)
-				SetScalar(Name + TEXT("_TexCoord"), TextureMap.TexCoord, 0.f);
+				SetScalar(Name + Inputs::PostFix::TexCoord, TextureMap.TexCoord, 0.f);
 
 				//Set the TilingMethod:
-				SetVec4(Name + TEXT("_TilingMethod"), GetTilingMethod(Textures[TextureMap.TextureIndex].Sampler), FVector4f(0, 0, 0, 0));
+				SetVec4(Name + Inputs::PostFix::TilingMethod, GetTilingMethod(Textures[TextureMap.TextureIndex].Sampler), FVector4f(0, 0, 0, 0));
 			}
 
 			if (TextureMap.bHasTextureTransform)
 			{
 				{//Material specific settings:
-					SetScalar(Name + TEXT("_Offset_X"), TextureMap.TextureTransform.Offset[0], 0.f, EProcessType::MATERIAL);
-					SetScalar(Name + TEXT("_Offset_Y"), TextureMap.TextureTransform.Offset[1], 0.f, EProcessType::MATERIAL);
+					SetScalar(Name + Inputs::PostFix::OffsetX, TextureMap.TextureTransform.Offset[0], 0.f, EProcessType::MATERIAL);
+					SetScalar(Name + Inputs::PostFix::OffsetY, TextureMap.TextureTransform.Offset[1], 0.f, EProcessType::MATERIAL);
 
-					SetScalar(Name + TEXT("_Scale_X"), TextureMap.TextureTransform.Scale[0], 1.f, EProcessType::MATERIAL);
-					SetScalar(Name + TEXT("_Scale_Y"), TextureMap.TextureTransform.Scale[1], 1.f, EProcessType::MATERIAL);
+					SetScalar(Name + Inputs::PostFix::ScaleX, TextureMap.TextureTransform.Scale[0], 1.f, EProcessType::MATERIAL);
+					SetScalar(Name + Inputs::PostFix::ScaleY, TextureMap.TextureTransform.Scale[1], 1.f, EProcessType::MATERIAL);
 				}
 
 				{//MaterialInstance specific settings:
 					FVector4f OffsetScale(TextureMap.TextureTransform.Offset[0], TextureMap.TextureTransform.Offset[1], TextureMap.TextureTransform.Scale[0], TextureMap.TextureTransform.Scale[1]);
-					SetVec4(Name + TEXT("_OffsetScale"), OffsetScale, FVector4f(0, 0, 1, 1), EProcessType::MATERIALINSTANCE);
+					SetVec4(Name + Inputs::PostFix::OffsetScale, OffsetScale, FVector4f(0, 0, 1, 1), EProcessType::MATERIALINSTANCE);
 				}
 
 				if (!FMath::IsNearlyZero(TextureMap.TextureTransform.Rotation))
@@ -398,7 +394,7 @@ namespace UE::Interchange::GLTFMaterials
 						AngleRadians = TWO_PI - AngleRadians;
 					}
 
-					SetScalar(Name + TEXT("_Rotation"), AngleRadians, 0.f);
+					SetScalar(Name + Inputs::PostFix::Rotation, AngleRadians, 0.f);
 				}
 			}
 		}
@@ -418,7 +414,7 @@ namespace UE::Interchange::GLTFMaterials
 				//BaseColorTexture_OffsetScale
 				//BaseColorTExture_Rotation
 				//BaseColorTExture_TexCoord
-				SetMap(GltfMaterial.BaseColor, Inputs::BaseColorTexture);
+				SetMap(Inputs::BaseColorTexture, GltfMaterial.BaseColor);
 
 				//BaseColorFactor
 				SetColor(Inputs::BaseColorFactor, GltfMaterial.BaseColorFactor, FVector4f(1, 1, 1, 1));
@@ -430,7 +426,7 @@ namespace UE::Interchange::GLTFMaterials
 				//NormalTexture_OffsetScale
 				//NormalTexture_Rotation
 				//NormalTexture_TexCoord
-				SetMap(GltfMaterial.Normal, Inputs::NormalTexture);
+				SetMap(Inputs::NormalTexture, GltfMaterial.Normal);
 
 				//NormalScale
 				SetScalar(Inputs::NormalScale, GltfMaterial.NormalScale, 1.f);
@@ -442,7 +438,7 @@ namespace UE::Interchange::GLTFMaterials
 				//EmissiveTexture_OffsetScale
 				//EmissiveTexture_Rotation
 				//EmissiveTexture_TexCoord
-				SetMap(GltfMaterial.Emissive, Inputs::EmissiveTexture);
+				SetMap(Inputs::EmissiveTexture, GltfMaterial.Emissive);
 
 				//EmissiveFactor
 				SetVec3(Inputs::EmissiveFactor, GltfMaterial.EmissiveFactor, FVector3f(0, 0, 0));
@@ -456,7 +452,7 @@ namespace UE::Interchange::GLTFMaterials
 				//OcclusionTexture_OffsetScale
 				//OcclusionTexture_Rotation
 				//OcclusionTexture_TexCoord
-				SetMap(GltfMaterial.Occlusion, Inputs::OcclusionTexture);
+				SetMap(Inputs::OcclusionTexture, GltfMaterial.Occlusion);
 
 				//OcclusionStrength
 				SetScalar(Inputs::OcclusionStrength, GltfMaterial.OcclusionStrength, 1.f);
@@ -473,7 +469,7 @@ namespace UE::Interchange::GLTFMaterials
 					//DiffuseTexture_OffsetScale
 					//DiffuseTexture_Rotation
 					//DiffuseTexture_TexCoord
-					SetMap(GltfMaterial.BaseColor, Inputs::DiffuseTexture);
+					SetMap(Inputs::DiffuseTexture, GltfMaterial.BaseColor);
 
 					//DiffuseFactor
 					SetColor(Inputs::DiffuseFactor, GltfMaterial.BaseColorFactor, FVector4f(1, 1, 1, 1));
@@ -484,7 +480,7 @@ namespace UE::Interchange::GLTFMaterials
 					//SpecularGlossinessTexture_OffsetScale
 					//SpecularGlossinessTexture_Rotation
 					//SpecularGlossinessTexture_TexCoord
-					SetMap(GltfMaterial.SpecularGlossiness.Map, Inputs::SpecularGlossinessTexture);
+					SetMap(Inputs::SpecularGlossinessTexture, GltfMaterial.SpecularGlossiness.Map);
 
 					//SpecFactor
 					FVector3f SpecularFactor(GltfMaterial.SpecularGlossiness.SpecularFactor[0], GltfMaterial.SpecularGlossiness.SpecularFactor[1], GltfMaterial.SpecularGlossiness.SpecularFactor[2]);
@@ -503,7 +499,7 @@ namespace UE::Interchange::GLTFMaterials
 					//BaseColorTexture_OffsetScale
 					//BaseColorTexture_Rotation
 					//BaseColorTexture_TexCoord
-					SetMap(GltfMaterial.BaseColor, Inputs::BaseColorTexture);
+					SetMap(Inputs::BaseColorTexture, GltfMaterial.BaseColor);
 
 					//BaseColorFactor
 					SetColor(Inputs::BaseColorFactor, GltfMaterial.BaseColorFactor, FVector4f(1, 1, 1, 1));
@@ -514,7 +510,7 @@ namespace UE::Interchange::GLTFMaterials
 					//MetallicRoughnessTexture_OffsetScale
 					//MetallicRoughnessTexture_Rotation
 					//MetallicRoughnessTexture_TexCoord
-					SetMap(GltfMaterial.MetallicRoughness.Map, Inputs::MetallicRoughnessTexture);
+					SetMap(Inputs::MetallicRoughnessTexture, GltfMaterial.MetallicRoughness.Map);
 
 					//MetallicFactor
 					SetScalar(Inputs::MetallicFactor, GltfMaterial.MetallicRoughness.MetallicFactor, 1.f);
@@ -529,7 +525,7 @@ namespace UE::Interchange::GLTFMaterials
 					//SpecularTexture_OffsetScale
 					//SpecularTexture_Rotation
 					//SpecularTexture_TexCoord
-					SetMap(GltfMaterial.Specular.SpecularMap, Inputs::SpecularTexture);
+					SetMap(Inputs::SpecularTexture, GltfMaterial.Specular.SpecularMap);
 
 					//SpecularFactor
 					SetScalar(Inputs::SpecularFactor, GltfMaterial.Specular.SpecularFactor, 0.5); //(UE Specular default is 0.5).
@@ -553,7 +549,7 @@ namespace UE::Interchange::GLTFMaterials
 					//TransmissionTexture_OffsetScale
 					//TransmissionTexture_Rotation
 					//TransmissionTexture_TexCoord
-					SetMap(GltfMaterial.Transmission.TransmissionMap, Inputs::TransmissionTexture);
+					SetMap(Inputs::TransmissionTexture, GltfMaterial.Transmission.TransmissionMap);
 
 					//TransmissionFactor
 					SetScalar(Inputs::TransmissionFactor, GltfMaterial.Transmission.TransmissionFactor, 0.f);
@@ -569,7 +565,7 @@ namespace UE::Interchange::GLTFMaterials
 					//ClearCoatTexture_OffsetScale
 					//ClearCoatTexture_Rotation
 					//ClearCoatTexture_TexCoord
-					SetMap(GltfMaterial.ClearCoat.ClearCoatMap, Inputs::ClearCoatTexture);
+					SetMap(Inputs::ClearCoatTexture, GltfMaterial.ClearCoat.ClearCoatMap);
 
 					//ClearCoatFactor
 					SetScalar(Inputs::ClearCoatFactor, GltfMaterial.ClearCoat.ClearCoatFactor, 0.f);
@@ -580,7 +576,7 @@ namespace UE::Interchange::GLTFMaterials
 					//ClearCoatRoughnessTexture_OffsetScale
 					//ClearCoatRoughnessTexture_Rotation
 					//ClearCoatRoughnessTexture_TexCoord
-					SetMap(GltfMaterial.ClearCoat.RoughnessMap, Inputs::ClearCoatRoughnessTexture);
+					SetMap(Inputs::ClearCoatRoughnessTexture, GltfMaterial.ClearCoat.RoughnessMap);
 
 					//ClearCoatRoughnessFactor
 					SetScalar(Inputs::ClearCoatRoughnessFactor, GltfMaterial.ClearCoat.Roughness, 0.f);
@@ -591,7 +587,7 @@ namespace UE::Interchange::GLTFMaterials
 					//ClearCoatNormalTexture_OffsetScale
 					//ClearCoatNormalTexture_Rotation
 					//ClearCoatNormalTexture_TexCoord
-					SetMap(GltfMaterial.ClearCoat.NormalMap, Inputs::ClearCoatNormalTexture);
+					SetMap(Inputs::ClearCoatNormalTexture, GltfMaterial.ClearCoat.NormalMap);
 
 					//ClearCoatNormalFactor
 					SetScalar(Inputs::ClearCoatNormalScale, GltfMaterial.ClearCoat.NormalMapUVScale, 1.f);
@@ -604,7 +600,7 @@ namespace UE::Interchange::GLTFMaterials
 					//SheenColorTexture_OffsetScale
 					//SheenColorTexture_Rotation
 					//SheenColorTexture_TexCoord
-					SetMap(GltfMaterial.Sheen.SheenColorMap, Inputs::SheenColorTexture);
+					SetMap(Inputs::SheenColorTexture, GltfMaterial.Sheen.SheenColorMap);
 
 					//SheenColorFactor
 					FVector3f SheenColorFactor(GltfMaterial.Sheen.SheenColorFactor[0], GltfMaterial.Sheen.SheenColorFactor[1], GltfMaterial.Sheen.SheenColorFactor[2]);
@@ -614,7 +610,7 @@ namespace UE::Interchange::GLTFMaterials
 					//SheenRoughnessTexture_OffsetScale
 					//SheenRoughnessTexture_Rotation
 					//SheenRoughnessTexture_TexCoord
-					SetMap(GltfMaterial.Sheen.SheenRoughnessMap, Inputs::SheenRoughnessTexture);
+					SetMap(Inputs::SheenRoughnessTexture, GltfMaterial.Sheen.SheenRoughnessMap);
 
 					//SheenRoughnessFactor
 					SetScalar(Inputs::SheenRoughnessFactor, GltfMaterial.Sheen.SheenRoughnessFactor, 0.f);
diff --git a/Engine/Plugins/Interchange/Runtime/Source/Import/Public/Gltf/InterchangeGLTFMaterial.h b/Engine/Plugins/Interchange/Runtime/Source/Import/Public/Gltf/InterchangeGLTFMaterial.h
--- a/Engine/Plugins/Interchange/Runtime/Source/Import/Public/Gltf/InterchangeGLTFMaterial.h
+++ b/Engine/Plugins/Interchange/Runtime/Source/Import/Public/Gltf/InterchangeGLTFMaterial.h
@@ -30,14 +30,31 @@ namespace UE::Interchange::GLTFMaterials
 	//Inputs/Parameters (Materials/MaterialInstances)
 	namespace Inputs
 	{
-		const FString Color_RGB = TEXT("_RGB");
-		const FString Color_A = TEXT("_A");
+		///PostFixes
+		namespace PostFix
+		{
+			const FString Color_RGB = TEXT("_RGB");
+			const FString Color_A = TEXT("_A");
+
+			const FString TexCoord = TEXT("_TexCoord");
+
+			const FString OffsetX = TEXT("_Offset_X");
+			const FString OffsetY = TEXT("_Offset_Y");
+			const FString ScaleX = TEXT("_Scale_X");
+			const FString ScaleY = TEXT("_Scale_Y");
+
+			const FString OffsetScale = TEXT("_OffsetScale");
+
+			const FString Rotation = TEXT("_Rotation");
+
+			const FString TilingMethod = TEXT("_TilingMethod");
+		}
 
 		//MetalRoughness specific:
 		DECLARE_INTERCHANGE_GLTF_MI_MAP(BaseColor)
 		const FString BaseColorFactor = TEXT("BaseColorFactor");
-		const FString BaseColorFactor_RGB = BaseColorFactor + Color_RGB; //Connection to inputs from BaseColorFactor.RGB
-		const FString BaseColorFactor_A = BaseColorFactor + Color_A; //Connection to inputs from BaseColorFactor.A
+		const FString BaseColorFactor_RGB = BaseColorFactor + PostFix::Color_RGB; //Connection to inputs from BaseColorFactor.RGB
+		const FString BaseColorFactor_A = BaseColorFactor + PostFix::Color_A; //Connection to inputs from BaseColorFactor.A
 
 		DECLARE_INTERCHANGE_GLTF_MI_MAP(MetallicRoughness)
 		const FString MetallicFactor = TEXT("MetallicFactor");
@@ -50,8 +67,8 @@ namespace UE::Interchange::GLTFMaterials
 		//SpecularGlossiness specific
 		DECLARE_INTERCHANGE_GLTF_MI_MAP(Diffuse)
 		const FString DiffuseFactor = TEXT("DiffuseFactor");
-		const FString DiffuseFactor_RGB = DiffuseFactor + Color_RGB; //Connection to inputs from BaseColorFactor.RGB
-		const FString DiffuseFactor_A = DiffuseFactor + Color_A; //Connection to inputs from BaseColorFactor.A
+		const FString DiffuseFactor_RGB = DiffuseFactor + PostFix::Color_RGB; //Connection to inputs from BaseColorFactor.RGB
+		const FString DiffuseFactor_A = DiffuseFactor + PostFix::Color_A; //Connection to inputs from BaseColorFactor.A
 
 		DECLARE_INTERCHANGE_GLTF_MI_MAP(SpecularGlossiness)
 		const FString SpecFactor = TEXT("SpecFactor");
@@ -116,11 +133,28 @@ namespace UE::Interchange::GLTFMaterials
 		Blend
 	};
 
-	static TMap<EShadingModel, TPair<FString, TArray<FString>>> ShadingModelToMaterialFunctions = {
+	struct FGLTFMaterialInformation
+	{
+		FString MaterialFunctionPath;
+		FString MaterialPath;
+		TArray<FString> MaterialFunctionOutputs;
+
+		FGLTFMaterialInformation(const FString& InMaterialFunctionPath,
+			const FString& InMaterialPath,
+			const TArray<FString>& InMaterialFunctionOutputs)
+			: MaterialFunctionPath(InMaterialFunctionPath)
+			, MaterialPath(InMaterialPath)
+			, MaterialFunctionOutputs(InMaterialFunctionOutputs)
+		{
+		}
+	};
+
+	static const TMap<EShadingModel, FGLTFMaterialInformation> ShadingModelToMaterialInformation = {
 		{EShadingModel::DEFAULT,
-			TPair<FString, TArray<FString>>(
-				TEXT("/Interchange/gltf/MaterialBodies/MF_Default_Body.MF_Default_Body"),
-				TArray<FString>{
+		FGLTFMaterialInformation(
+			TEXT("/Interchange/gltf/MaterialBodies/MF_Default_Body.MF_Default_Body"),
+			TEXT("/Interchange/gltf/M_Default.M_Default"),
+			TArray<FString>{
 					TEXT("BaseColor"),
 					TEXT("Metallic"),
 					TEXT("Specular"),
@@ -129,20 +163,24 @@ namespace UE::Interchange::GLTFMaterials
 					TEXT("Opacity"),
 					TEXT("OpacityMask"),
 					TEXT("Normal"),
-					TEXT("Occlusion")})},
+					TEXT("Occlusion")})
+		},
 
 		{EShadingModel::UNLIT,
-			TPair<FString, TArray<FString>>(
-				TEXT("/Interchange/gltf/MaterialBodies/MF_Unlit_Body.MF_Unlit_Body"),
-				TArray<FString>{
+		FGLTFMaterialInformation(
+			TEXT("/Interchange/gltf/MaterialBodies/MF_Unlit_Body.MF_Unlit_Body"),
+			TEXT("/Interchange/gltf/M_Unlit.M_Unlit"),
+			TArray<FString>{
 					TEXT("UnlitColor"),
 					TEXT("Opacity"),
-					TEXT("OpacityMask")})},
+					TEXT("OpacityMask")})
+		},
 
 		{EShadingModel::CLEARCOAT,
-			TPair<FString, TArray<FString>>(
-				TEXT("/Interchange/gltf/MaterialBodies/MF_ClearCoat_Body.MF_ClearCoat_Body"),
-				TArray<FString>{
+		FGLTFMaterialInformation(
+			TEXT("/Interchange/gltf/MaterialBodies/MF_ClearCoat_Body.MF_ClearCoat_Body"),
+			TEXT("/Interchange/gltf/M_ClearCoat.M_ClearCoat"),
+			TArray<FString>{
 					TEXT("ClearCoatNormal"),
 					TEXT("BaseColor"),
 					TEXT("Metallic"),
@@ -154,12 +192,14 @@ namespace UE::Interchange::GLTFMaterials
 					TEXT("Normal"),
 					TEXT("ClearCoat"),
 					TEXT("ClearCoatRoughness"),
-					TEXT("Occlusion")})},
+					TEXT("Occlusion")})
+		},
 
 		{EShadingModel::SHEEN,
-			TPair<FString, TArray<FString>>(
-				TEXT("/Interchange/gltf/MaterialBodies/MF_Sheen_Body.MF_Sheen_Body"),
-				TArray<FString>{
+		FGLTFMaterialInformation(
+			TEXT("/Interchange/gltf/MaterialBodies/MF_Sheen_Body.MF_Sheen_Body"),
+			TEXT("/Interchange/gltf/M_Sheen.M_Sheen"),
+			TArray<FString>{
 					TEXT("BaseColor"),
 					TEXT("Metallic"),
 					TEXT("Specular"),
@@ -170,12 +210,14 @@ namespace UE::Interchange::GLTFMaterials
 					TEXT("Normal"),
 					TEXT("SheenColor"),
 					TEXT("SheenRoughness"),
-					TEXT("Occlusion")})},
+					TEXT("Occlusion")})
+		},
 
 		{EShadingModel::TRANSMISSION,
-			TPair<FString, TArray<FString>>(
-				TEXT("/Interchange/gltf/MaterialBodies/MF_Transmission_Body.MF_Transmission_Body"),
-				TArray<FString>{
+		FGLTFMaterialInformation(
+			TEXT("/Interchange/gltf/MaterialBodies/MF_Transmission_Body.MF_Transmission_Body"),
+			TEXT("/Interchange/gltf/M_Transmission.M_Transmission"),
+			TArray<FString>{
 					TEXT("TransmissionColor"),
 					TEXT("BaseColor"),
 					TEXT("Metallic"),
@@ -184,12 +226,14 @@ namespace UE::Interchange::GLTFMaterials
 					TEXT("EmissiveColor"),
 					TEXT("Opacity"),
 					TEXT("Normal"),
-					TEXT("Occlusion")})},
+					TEXT("Occlusion")})
+		},
 
 		{EShadingModel::SPECULARGLOSSINESS,
-			TPair<FString, TArray<FString>>(
-				TEXT("/Interchange/gltf/MaterialBodies/MF_SpecularGlossiness_Body.MF_SpecularGlossiness_Body"),
-				TArray<FString>{
+		FGLTFMaterialInformation(
+			TEXT("/Interchange/gltf/MaterialBodies/MF_SpecularGlossiness_Body.MF_SpecularGlossiness_Body"),
+			TEXT("/Interchange/gltf/M_SpecularGlossiness.M_SpecularGlossiness"),
+			TArray<FString>{
 					TEXT("BaseColor"),
 					TEXT("Metallic"),
 					TEXT("Roughness"),
@@ -197,15 +241,36 @@ namespace UE::Interchange::GLTFMaterials
 					TEXT("Opacity"),
 					TEXT("OpacityMask"),
 					TEXT("Normal"),
-					TEXT("Occlusion")})},
+					TEXT("Occlusion")})
+		}
 	};
 
 	inline TArray<FString> GetRequiredMaterialFunctionPaths()
 	{
 		TArray<FString> Result;
-		for (const TPair<EShadingModel, TPair<FString, TArray<FString>>>& Entry : ShadingModelToMaterialFunctions)
+		for (const TPair<EShadingModel, FGLTFMaterialInformation>& Entry : ShadingModelToMaterialInformation)
+		{
+			Result.Add(Entry.Value.MaterialFunctionPath);
+		}
+		return Result;
+	}
+
+	inline TMap<FString, EShadingModel> GetMaterialFunctionPathsToShadingModels()
+	{
+		TMap<FString, EShadingModel> Result;
+		for (const TPair<EShadingModel, FGLTFMaterialInformation>& Entry : ShadingModelToMaterialInformation)
+		{
+			Result.Add(Entry.Value.MaterialFunctionPath, Entry.Key);
+		}
+		return Result;
+	}
+
+	inline TMap<FString, EShadingModel> GetMaterialPathsToShadingModels()
+	{
+		TMap<FString, EShadingModel> Result;
+		for (const TPair<EShadingModel, FGLTFMaterialInformation>& Entry : ShadingModelToMaterialInformation)
 		{
-			Result.Add(Entry.Value.Key);
+			Result.Add(Entry.Value.MaterialPath, Entry.Key);
 		}
 		return Result;
 	}
diff --git a/Engine/Plugins/Interchange/Runtime/Source/Pipelines/Private/InterchangeglTFPipeline.cpp b/Engine/Plugins/Interchange/Runtime/Source/Pipelines/Private/InterchangeglTFPipeline.cpp
--- a/Engine/Plugins/Interchange/Runtime/Source/Pipelines/Private/InterchangeglTFPipeline.cpp
+++ b/Engine/Plugins/Interchange/Runtime/Source/Pipelines/Private/InterchangeglTFPipeline.cpp
@@ -92,10 +92,10 @@ TSet<FString> UGLTFPipelineSettings::GenerateExpectedParametersList(const FStrin
 	auto AddTextureAndRelated = [&ExpectedParameters](const FString& TextureName)
 		{
 			ExpectedParameters.Add(TextureName);
-			ExpectedParameters.Add(TextureName + TEXT("_OffsetScale"));
-			ExpectedParameters.Add(TextureName + TEXT("_Rotation"));
-			ExpectedParameters.Add(TextureName + TEXT("_TexCoord"));
-			ExpectedParameters.Add(TextureName + TEXT("_TilingMethod"));
+			ExpectedParameters.Add(TextureName + Inputs::PostFix::OffsetScale);
+			ExpectedParameters.Add(TextureName + Inputs::PostFix::Rotation);
+			ExpectedParameters.Add(TextureName + Inputs::PostFix::TexCoord);
+			ExpectedParameters.Add(TextureName + Inputs::PostFix::TilingMethod);
 		};
 
 	if (Identifier.Contains(TEXT("_Unlit")))
