[glTF Exporter] - GLTFExporter removed references to the Private directory of the MeshMergeUtilities module.

#rnx

diff --git a/Engine/Plugins/Enterprise/GLTFExporter/Source/GLTFExporter/GLTFExporter.Build.cs b/Engine/Plugins/Enterprise/GLTFExporter/Source/GLTFExporter/GLTFExporter.Build.cs
--- a/Engine/Plugins/Enterprise/GLTFExporter/Source/GLTFExporter/GLTFExporter.Build.cs
+++ b/Engine/Plugins/Enterprise/GLTFExporter/Source/GLTFExporter/GLTFExporter.Build.cs
@@ -78,12 +78,9 @@ public class GLTFExporter : ModuleRules
 					"MeshMergeUtilities",
 					"MeshDescription",
 					"StaticMeshDescription",
+					"MeshMergeUtilities",
 				}
 			);
-
-			PrivateIncludePaths.AddRange(new string[] {
-				System.IO.Path.Combine(GetModuleDirectory("MeshMergeUtilities"), "Private"),
-			});
 		}
 	}
 }
diff --git a/Engine/Plugins/Enterprise/GLTFExporter/Source/GLTFExporter/Private/Converters/GLTFMeshData.cpp b/Engine/Plugins/Enterprise/GLTFExporter/Source/GLTFExporter/Private/Converters/GLTFMeshData.cpp
--- a/Engine/Plugins/Enterprise/GLTFExporter/Source/GLTFExporter/Private/Converters/GLTFMeshData.cpp
+++ b/Engine/Plugins/Enterprise/GLTFExporter/Source/GLTFExporter/Private/Converters/GLTFMeshData.cpp
@@ -10,10 +10,12 @@
 #if WITH_EDITOR
 #include "Editor.h"
 #include "Engine/MapBuildDataRegistry.h"
-#include "MeshMergeHelpers.h"
 #include "StaticMeshComponentLODInfo.h"
 #include "Rendering/SkeletalMeshRenderData.h"
 #include "StaticMeshAttributes.h"
+#include "IMeshMergeUtilities.h"
+#include "MeshMergeModule.h"
+#include "Modules/ModuleManager.h"
 #endif
 
 FGLTFMeshData::FGLTFMeshData(const UStaticMesh* StaticMesh, const UStaticMeshComponent* StaticMeshComponent, int32 LODIndex)
@@ -26,6 +28,7 @@ FGLTFMeshData::FGLTFMeshData(const UStaticMesh* StaticMesh, const UStaticMeshCom
 #endif
 {
 #if WITH_EDITOR
+	const IMeshMergeUtilities& MeshMergeUtilities = FModuleManager::Get().LoadModuleChecked<IMeshMergeModule>("MeshMergeUtilities").GetUtilities();
 	FStaticMeshAttributes(Description).Register();
 #endif
 
@@ -35,7 +38,10 @@ FGLTFMeshData::FGLTFMeshData(const UStaticMesh* StaticMesh, const UStaticMeshCom
 
 #if WITH_EDITOR
 		PrimitiveData = { StaticMeshComponent };
-		FMeshMergeHelpers::RetrieveMesh(StaticMeshComponent, LODIndex, Description, true, false);
+
+		//MeshMergeUtilities.RetrieveMeshDescription (for StaticMeshComponent) uses bApplyComponentTransform==true, 
+		// however MeshData->Description is only used for DegenerateVertices/Triangles and UV checks, and material baking.
+		MeshMergeUtilities.RetrieveMeshDescription(StaticMeshComponent, LODIndex, Description, true);
 
 		constexpr int32 LightMapLODIndex = 0; // TODO: why is this zero?
 		if (StaticMeshComponent->LODData.IsValidIndex(LightMapLODIndex))
@@ -56,7 +62,7 @@ FGLTFMeshData::FGLTFMeshData(const UStaticMesh* StaticMesh, const UStaticMeshCom
 
 #if WITH_EDITOR
 		PrimitiveData = { StaticMesh };
-		FMeshMergeHelpers::RetrieveMesh(StaticMesh, LODIndex, Description);
+		MeshMergeUtilities.RetrieveMeshDescription(StaticMesh, LODIndex, Description);
 #endif
 	}
 
@@ -78,6 +84,7 @@ FGLTFMeshData::FGLTFMeshData(const USkeletalMesh* SkeletalMesh, const USkeletalM
 #endif
 {
 #if WITH_EDITOR
+	const IMeshMergeUtilities& MeshMergeUtilities = FModuleManager::Get().LoadModuleChecked<IMeshMergeModule>("MeshMergeUtilities").GetUtilities();
 	FStaticMeshAttributes(Description).Register();
 #endif
 
@@ -87,7 +94,7 @@ FGLTFMeshData::FGLTFMeshData(const USkeletalMesh* SkeletalMesh, const USkeletalM
 
 #if WITH_EDITOR
 		PrimitiveData = { SkeletalMeshComponent };
-		FMeshMergeHelpers::RetrieveMesh(SkeletalMeshComponent, LODIndex, Description, true, false);
+		MeshMergeUtilities.RetrieveMeshDescription(SkeletalMeshComponent, LODIndex, Description, true);
 #endif
 	}
 	else
@@ -115,7 +122,7 @@ FGLTFMeshData::FGLTFMeshData(const USkeletalMesh* SkeletalMesh, const USkeletalM
 				TempComponent->RegisterComponent();
 				TempComponent->SetSkeletalMesh(const_cast<USkeletalMesh*>(SkeletalMesh));
 
-				FMeshMergeHelpers::RetrieveMesh(TempComponent, LODIndex, Description, true, false);
+				MeshMergeUtilities.RetrieveMeshDescription(TempComponent, LODIndex, Description, true);
 
 				World->DestroyActor(TempActor, false, false);
 			}
