[glTFExporter] Landscape, Spline, InstancedMesh support.

#rnx

diff --git a/Engine/Plugins/Enterprise/GLTFExporter/Source/GLTFExporter/GLTFExporter.Build.cs b/Engine/Plugins/Enterprise/GLTFExporter/Source/GLTFExporter/GLTFExporter.Build.cs
--- a/Engine/Plugins/Enterprise/GLTFExporter/Source/GLTFExporter/GLTFExporter.Build.cs
+++ b/Engine/Plugins/Enterprise/GLTFExporter/Source/GLTFExporter/GLTFExporter.Build.cs
@@ -47,6 +47,7 @@ public class GLTFExporter : ModuleRules
 				"Projects",
 				"EngineSettings",
 				"Analytics",
+				"Landscape",
 			}
 		);
 
diff --git a/Engine/Plugins/Enterprise/GLTFExporter/Source/GLTFExporter/Private/Builders/GLTFAnalyticsBuilder.cpp b/Engine/Plugins/Enterprise/GLTFExporter/Source/GLTFExporter/Private/Builders/GLTFAnalyticsBuilder.cpp
--- a/Engine/Plugins/Enterprise/GLTFExporter/Source/GLTFExporter/Private/Builders/GLTFAnalyticsBuilder.cpp
+++ b/Engine/Plugins/Enterprise/GLTFExporter/Source/GLTFExporter/Private/Builders/GLTFAnalyticsBuilder.cpp
@@ -19,13 +19,15 @@ TArray<FAnalyticsEventAttribute> FGLTFAnalyticsBuilder::GenerateAnalytics() cons
 	int32 NumberOfActors = ActorsRecorded.Num();
 	int32 NumberOfStaticMeshes = StaticMeshesRecorded.Num();
 	int32 NumberOfSkeletalMeshes = SkeletalMeshesRecorded.Num();
+	int32 NumberOfSplineStaticMeshes = SplineStaticMeshesRecorded.Num();
+	int32 NumberOfLandscapes = LandscapeComponentsRecorded.Num();
 	int32 NumberOfLevelSequences = LevelSequencesRecorded.Num();
 	int32 NumberOfAnimSequences = AnimSequencesRecorded.Num();
 	int32 NumberOfMaterials = MaterialsRecorded.Num();
 	int32 NumberOfTextures = TexturesRecorded.Num();
 	int32 NumberOfCameras = CamerasRecorded.Num();
 	int32 NumberOfLights = LightsRecorded.Num();
-
+	
 	TArray<FAnalyticsEventAttribute> EventAttributes;
 	auto Add = [&EventAttributes](const FString& Name, int32 Value)
 	{
@@ -35,6 +37,8 @@ TArray<FAnalyticsEventAttribute> FGLTFAnalyticsBuilder::GenerateAnalytics() cons
 	Add(TEXT("NumberOfActors"), NumberOfActors);
 	Add(TEXT("NumberOfStaticMeshes"), NumberOfStaticMeshes);
 	Add(TEXT("NumberOfSkeletalMeshes"), NumberOfSkeletalMeshes);
+	Add(TEXT("NumberOfSplineStaticMeshes"), NumberOfSplineStaticMeshes);
+	Add(TEXT("NumberOfLandscapes"), NumberOfLandscapes);
 	Add(TEXT("NumberOfLevelSequences"), NumberOfLevelSequences);
 	Add(TEXT("NumberOfAnimSequences"), NumberOfAnimSequences);
 	Add(TEXT("NumberOfMaterials"), NumberOfMaterials);
@@ -106,6 +110,15 @@ void FGLTFAnalyticsBuilder::RecordSkeletalMesh(const USkeletalMesh* Object)
 	SkeletalMeshesRecorded.Add(Object);
 }
 
+void FGLTFAnalyticsBuilder::RecordSplineStaticMesh(const UStaticMesh* Object)
+{
+	SplineStaticMeshesRecorded.Add(Object);
+}
+void FGLTFAnalyticsBuilder::RecordLandscapeComponent(const ULandscapeComponent* Object)
+{
+	LandscapeComponentsRecorded.Add(Object);
+}
+
 void FGLTFAnalyticsBuilder::RecordLevelSequence(const ULevelSequence* Object)
 {
 	LevelSequencesRecorded.Add(Object);
diff --git a/Engine/Plugins/Enterprise/GLTFExporter/Source/GLTFExporter/Private/Builders/GLTFConvertBuilder.cpp b/Engine/Plugins/Enterprise/GLTFExporter/Source/GLTFExporter/Private/Builders/GLTFConvertBuilder.cpp
--- a/Engine/Plugins/Enterprise/GLTFExporter/Source/GLTFExporter/Private/Builders/GLTFConvertBuilder.cpp
+++ b/Engine/Plugins/Enterprise/GLTFExporter/Source/GLTFExporter/Private/Builders/GLTFConvertBuilder.cpp
@@ -6,6 +6,10 @@
 #include "Components/StaticMeshComponent.h"
 #include "Engine/Texture2D.h"
 #include "Engine/TextureRenderTarget2D.h"
+#include "LandscapeComponent.h"
+#include "Components/SplineMeshComponent.h"
+#include "Converters/GLTFNormalArray.h"
+#include "Converters/GLTFUVArray.h"
 
 FGLTFConvertBuilder::FGLTFConvertBuilder(const FString& FileName, const UGLTFExportOptions* ExportOptions, const TSet<AActor*>& SelectedActors)
 	: FGLTFAnalyticsBuilder(FileName, ExportOptions)
@@ -34,6 +38,16 @@ FGLTFJsonAccessor* FGLTFConvertBuilder::AddUniquePositionAccessor(const FGLTFMes
 	return PositionBufferConverter->GetOrAdd(MeshSection, VertexBuffer);
 }
 
+FGLTFJsonAccessor* FGLTFConvertBuilder::AddUniquePositionAccessor(const FPositionVertexBuffer* VertexBuffer)
+{
+	if (VertexBuffer == nullptr)
+	{
+		return nullptr;
+	}
+
+	return PositionBufferConverterRaw->GetOrAdd(VertexBuffer);
+}
+
 FGLTFJsonAccessor* FGLTFConvertBuilder::AddUniqueColorAccessor(const FGLTFMeshSection* MeshSection, const FColorVertexBuffer* VertexBuffer)
 {
 	if (VertexBuffer == nullptr)
@@ -54,6 +68,16 @@ FGLTFJsonAccessor* FGLTFConvertBuilder::AddUniqueNormalAccessor(const FGLTFMeshS
 	return NormalBufferConverter->GetOrAdd(MeshSection, VertexBuffer);
 }
 
+FGLTFJsonAccessor* FGLTFConvertBuilder::AddUniqueNormalAccessor(const FGLTFNormalArray* Normals)
+{
+	if (Normals == nullptr)
+	{
+		return nullptr;
+	}
+
+	return NormalBufferConverterRaw->GetOrAdd(Normals);
+}
+
 FGLTFJsonAccessor* FGLTFConvertBuilder::AddUniqueTangentAccessor(const FGLTFMeshSection* MeshSection, const FStaticMeshVertexBuffer* VertexBuffer)
 {
 	if (VertexBuffer == nullptr)
@@ -74,6 +98,16 @@ FGLTFJsonAccessor* FGLTFConvertBuilder::AddUniqueUVAccessor(const FGLTFMeshSecti
 	return UVBufferConverter->GetOrAdd(MeshSection, VertexBuffer, UVIndex);
 }
 
+FGLTFJsonAccessor* FGLTFConvertBuilder::AddUniqueUVAccessor(const FGLTFUVArray* UVs)
+{
+	if (UVs == nullptr)
+	{
+		return nullptr;
+	}
+
+	return UVBufferConverterRaw->GetOrAdd(UVs);
+}
+
 FGLTFJsonAccessor* FGLTFConvertBuilder::AddUniqueJointAccessor(const FGLTFMeshSection* MeshSection, const FSkinWeightVertexBuffer* VertexBuffer, int32 InfluenceOffset)
 {
 	if (VertexBuffer == nullptr)
@@ -104,6 +138,16 @@ FGLTFJsonAccessor* FGLTFConvertBuilder::AddUniqueIndexAccessor(const FGLTFMeshSe
 	return IndexBufferConverter->GetOrAdd(MeshSection);
 }
 
+FGLTFJsonAccessor* FGLTFConvertBuilder::AddUniqueIndexAccessor(const FGLTFIndexArray* IndexBuffer, const FString& MeshName)
+{
+	if (IndexBuffer == nullptr)
+	{
+		return nullptr;
+	}
+
+	return IndexBufferConverterRaw->GetOrAdd(IndexBuffer, MeshName);
+}
+
 FGLTFJsonMesh* FGLTFConvertBuilder::AddUniqueMesh(const UStaticMesh* StaticMesh, const FGLTFMaterialArray& Materials, int32 LODIndex)
 {
 	if (StaticMesh == nullptr)
@@ -179,6 +223,36 @@ FGLTFJsonMesh* FGLTFConvertBuilder::AddUniqueMesh(const USkeletalMeshComponent*
 	return SkeletalMeshConverter->GetOrAdd(SkeletalMesh, SkeletalMeshComponent, Materials, LODIndex);
 }
 
+FGLTFJsonMesh* FGLTFConvertBuilder::AddUniqueMesh(const ULandscapeComponent* LandscapeComponent, const UMaterialInterface* LandscapeMaterial)
+{
+	if (LandscapeComponent == nullptr)
+	{
+		return nullptr;
+	}
+
+	RecordLandscapeComponent(LandscapeComponent);
+
+	return LandscapeConverter->GetOrAdd(LandscapeComponent, LandscapeMaterial);
+}
+
+FGLTFJsonMesh* FGLTFConvertBuilder::AddUniqueMesh(const USplineMeshComponent* SplineMeshComponent, const FGLTFMaterialArray& Materials, int32 LODIndex)
+{
+	if (SplineMeshComponent == nullptr)
+	{
+		return nullptr;
+	}
+
+	const UStaticMesh* StaticMesh = SplineMeshComponent->GetStaticMesh();
+	if (StaticMesh == nullptr)
+	{
+		return nullptr;
+	}
+
+	RecordSplineStaticMesh(StaticMesh);
+
+	return SplineMeshConverter->GetOrAdd(StaticMesh, SplineMeshComponent, Materials, LODIndex);
+}
+
 const FGLTFMeshData* FGLTFConvertBuilder::AddUniqueMeshData(const UStaticMesh* StaticMesh, const UStaticMeshComponent* StaticMeshComponent, int32 LODIndex)
 {
 	RecordStaticMesh(StaticMesh);
@@ -522,4 +596,4 @@ void FGLTFConvertBuilder::RegisterObjectVariant(const UObject* Object, const UPr
 const TArray<const UPropertyValue*>* FGLTFConvertBuilder::GetObjectVariants(const UObject* Object) const
 {
 	return ObjectVariants.Find(Object);
-}
+}
\ No newline at end of file
diff --git a/Engine/Plugins/Enterprise/GLTFExporter/Source/GLTFExporter/Private/Converters/GLTFAccessorConverters.cpp b/Engine/Plugins/Enterprise/GLTFExporter/Source/GLTFExporter/Private/Converters/GLTFAccessorConverters.cpp
--- a/Engine/Plugins/Enterprise/GLTFExporter/Source/GLTFExporter/Private/Converters/GLTFAccessorConverters.cpp
+++ b/Engine/Plugins/Enterprise/GLTFExporter/Source/GLTFExporter/Private/Converters/GLTFAccessorConverters.cpp
@@ -8,6 +8,8 @@
 #include "Rendering/PositionVertexBuffer.h"
 #include "Rendering/SkinWeightVertexBuffer.h"
 #include "Rendering/StaticMeshVertexBuffer.h"
+#include "Converters/GLTFNormalArray.h"
+#include "Converters/GLTFUVArray.h"
 
 namespace
 {
@@ -46,67 +48,109 @@ namespace
 			return FGLTFCoreUtilities::ConvertTangent(TangentX.ToFVector3f(), TangentZ.ToFVector4f());
 		}
 	};
-}
 
-FGLTFJsonAccessor* FGLTFPositionBufferConverter::Convert(const FGLTFMeshSection* MeshSection, const FPositionVertexBuffer* VertexBuffer)
-{
-	if (VertexBuffer == nullptr || VertexBuffer->GetNumVertices() == 0)
+	FGLTFJsonAccessor* ConvertVertexBuffer(const FPositionVertexBuffer* VertexBuffer, FGLTFConvertBuilder& Builder, const TArray<uint32>& IndexMap = TArray<uint32>())
 	{
-		return nullptr;
-	}
+		if (VertexBuffer == nullptr || VertexBuffer->GetNumVertices() == 0)
+		{
+			return nullptr;
+		}
 
-	const TUniquePtr<IGLTFBufferAdapter> SourceBuffer = IGLTFBufferAdapter::GetPositions(VertexBuffer);
-	const uint8* SourceData = SourceBuffer->GetData();
+		const TUniquePtr<IGLTFBufferAdapter> SourceBuffer = IGLTFBufferAdapter::GetPositions(VertexBuffer);
+		const uint8* SourceData = SourceBuffer->GetData();
 
-	if (SourceData == nullptr)
-	{
-		// TODO: report error
-		return nullptr;
-	}
+		if (SourceData == nullptr)
+		{
+			// TODO: report error
+			return nullptr;
+		}
 
-	const TArray<uint32>& IndexMap = MeshSection->IndexMap;
-	const uint32 VertexCount = IndexMap.Num();
-	const uint32 Stride = VertexBuffer->GetStride();
+		const uint32 VertexCount = IndexMap.Num() > 0 ? IndexMap.Num() : VertexBuffer->GetNumVertices();
+		const uint32 Stride = VertexBuffer->GetStride();
 
-	TArray<FGLTFVector3> Positions;
-	Positions.AddUninitialized(VertexCount);
+		TArray<FGLTFVector3> Positions;
+		Positions.AddUninitialized(VertexCount);
 
-	for (uint32 VertexIndex = 0; VertexIndex < VertexCount; ++VertexIndex)
-	{
-		const uint32 MappedVertexIndex = IndexMap[VertexIndex];
-		const FVector3f& Position = *reinterpret_cast<const FVector3f*>(SourceData + Stride * MappedVertexIndex);
-		Positions[VertexIndex] = FGLTFCoreUtilities::ConvertPosition(Position, Builder.ExportOptions->ExportUniformScale);
-	}
+		if (IndexMap.Num() > 0)
+		{
+			for (uint32 VertexIndex = 0; VertexIndex < VertexCount; ++VertexIndex)
+			{
+				const uint32 MappedVertexIndex = IndexMap[VertexIndex];
+				const FVector3f& Position = *reinterpret_cast<const FVector3f*>(SourceData + Stride * MappedVertexIndex);
+				Positions[VertexIndex] = FGLTFCoreUtilities::ConvertPosition(Position, Builder.ExportOptions->ExportUniformScale);
+			}
+		}
+		else
+		{
+			for (uint32 VertexIndex = 0; VertexIndex < VertexCount; ++VertexIndex)
+			{
+				const FVector3f& Position = *reinterpret_cast<const FVector3f*>(SourceData + Stride * VertexIndex);
+				Positions[VertexIndex] = FGLTFCoreUtilities::ConvertPosition(Position, Builder.ExportOptions->ExportUniformScale);
+			}
+		}
 
-	FGLTFJsonAccessor* JsonAccessor = Builder.AddAccessor();
-	JsonAccessor->BufferView = Builder.AddBufferView(Positions, EGLTFJsonBufferTarget::ArrayBuffer);
-	JsonAccessor->ComponentType = EGLTFJsonComponentType::Float;
-	JsonAccessor->Count = VertexCount;
-	JsonAccessor->Type = EGLTFJsonAccessorType::Vec3;
 
-	if (VertexCount > 0)
-	{
-		// Calculate accurate bounding box based on raw vertex values
-		JsonAccessor->MinMaxLength = 3;
+		FGLTFJsonAccessor* JsonAccessor = Builder.AddAccessor();
+		JsonAccessor->BufferView = Builder.AddBufferView(Positions, EGLTFJsonBufferTarget::ArrayBuffer);
+		JsonAccessor->ComponentType = EGLTFJsonComponentType::Float;
+		JsonAccessor->Count = VertexCount;
+		JsonAccessor->Type = EGLTFJsonAccessorType::Vec3;
 
-		for (int32 ComponentIndex = 0; ComponentIndex < JsonAccessor->MinMaxLength; ComponentIndex++)
+		if (VertexCount > 0)
 		{
-			JsonAccessor->Min[ComponentIndex] = Positions[0].Components[ComponentIndex];
-			JsonAccessor->Max[ComponentIndex] = Positions[0].Components[ComponentIndex];
-		}
+			// Calculate accurate bounding box based on raw vertex values
+			JsonAccessor->MinMaxLength = 3;
 
-		for (uint32 VertexIndex = 1; VertexIndex < VertexCount; ++VertexIndex)
-		{
-			const FGLTFVector3& Position = Positions[VertexIndex];
 			for (int32 ComponentIndex = 0; ComponentIndex < JsonAccessor->MinMaxLength; ComponentIndex++)
 			{
-				JsonAccessor->Min[ComponentIndex] = FMath::Min(JsonAccessor->Min[ComponentIndex], Position.Components[ComponentIndex]);
-				JsonAccessor->Max[ComponentIndex] = FMath::Max(JsonAccessor->Max[ComponentIndex], Position.Components[ComponentIndex]);
+				JsonAccessor->Min[ComponentIndex] = Positions[0].Components[ComponentIndex];
+				JsonAccessor->Max[ComponentIndex] = Positions[0].Components[ComponentIndex];
+			}
+
+			for (uint32 VertexIndex = 1; VertexIndex < VertexCount; ++VertexIndex)
+			{
+				const FGLTFVector3& Position = Positions[VertexIndex];
+				for (int32 ComponentIndex = 0; ComponentIndex < JsonAccessor->MinMaxLength; ComponentIndex++)
+				{
+					JsonAccessor->Min[ComponentIndex] = FMath::Min(JsonAccessor->Min[ComponentIndex], Position.Components[ComponentIndex]);
+					JsonAccessor->Max[ComponentIndex] = FMath::Max(JsonAccessor->Max[ComponentIndex], Position.Components[ComponentIndex]);
+				}
 			}
 		}
+
+		return JsonAccessor;
 	}
 
-	return JsonAccessor;
+	template <typename IndexType, typename IndexBufferType>
+	FGLTFJsonAccessor* ConvertIndexBuffer(const TArray<IndexBufferType>& IndexBuffer, FGLTFConvertBuilder& Builder)
+	{
+		const uint32 IndexCount = IndexBuffer.Num();
+		if (IndexCount == 0)
+		{
+			return nullptr;
+		}
+
+		TArray<IndexType> Indices;
+		Indices.AddUninitialized(IndexCount);
+
+		for (uint32 Index = 0; Index < IndexCount; ++Index)
+		{
+			Indices[Index] = static_cast<IndexType>(IndexBuffer[Index]);
+		}
+
+		FGLTFJsonAccessor* JsonAccessor = Builder.AddAccessor();
+		JsonAccessor->BufferView = Builder.AddBufferView(Indices, EGLTFJsonBufferTarget::ElementArrayBuffer, sizeof(IndexType));
+		JsonAccessor->ComponentType = FGLTFCoreUtilities::GetComponentType<IndexType>();
+		JsonAccessor->Count = IndexCount;
+		JsonAccessor->Type = EGLTFJsonAccessorType::Scalar;
+
+		return JsonAccessor;
+	}
+}
+
+FGLTFJsonAccessor* FGLTFPositionBufferConverter::Convert(const FGLTFMeshSection* MeshSection, const FPositionVertexBuffer* VertexBuffer)
+{
+	return ConvertVertexBuffer(VertexBuffer, Builder, MeshSection->IndexMap);
 }
 
 FGLTFJsonAccessor* FGLTFColorBufferConverter::Convert(const FGLTFMeshSection* MeshSection, const FColorVertexBuffer* VertexBuffer)
@@ -570,45 +614,106 @@ FGLTFJsonAccessor* FGLTFBoneWeightBufferConverter::Convert(const FGLTFMeshSectio
 
 FGLTFJsonAccessor* FGLTFIndexBufferConverter::Convert(const FGLTFMeshSection* MeshSection)
 {
+	const TArray<uint32>& IndexBuffer = MeshSection->IndexBuffer;
 	const uint32 MaxVertexIndex = MeshSection->IndexMap.Num() - 1;
 	// NOTE: Even if a maximum value (i.e. UINT8_MAX or UINT16_MAX) would fit inside the type (i.e. uint8 or uint16), it's not allowed per gltf spec
 	// since the maximum values trigger primitive restart in some graphics APIs and would require client implementations to rebuild the index buffer.
-	if (MaxVertexIndex < UINT8_MAX) return Convert<uint8>(MeshSection);
-	if (MaxVertexIndex < UINT16_MAX) return Convert<uint16>(MeshSection);
+	if (MaxVertexIndex < UINT8_MAX) return ConvertIndexBuffer<uint8>(IndexBuffer, Builder);
+	if (MaxVertexIndex < UINT16_MAX) return ConvertIndexBuffer<uint16>(IndexBuffer, Builder);
 	if (MaxVertexIndex == UINT32_MAX)
 	{
-		Builder.LogError(
+		Builder.LogWarning(
 			FString::Printf(TEXT("Maximum value (%d) used in indices accessor for mesh %s, this may not be supported on some glTF viewers"),
-			UINT32_MAX,
-			*MeshSection->ToString()
+				UINT32_MAX,
+				*MeshSection->ToString()
 			));
 	}
-	return Convert<uint32>(MeshSection);
+	return ConvertIndexBuffer<uint32>(IndexBuffer, Builder);
 }
 
-template <typename IndexType>
-FGLTFJsonAccessor* FGLTFIndexBufferConverter::Convert(const FGLTFMeshSection* MeshSection) const
+
+FGLTFJsonAccessor* FGLTFPositionBufferConverterRaw::Convert(const FPositionVertexBuffer* VertexBuffer)
 {
-	const TArray<uint32>& IndexBuffer = MeshSection->IndexBuffer;
-	const uint32 IndexCount = IndexBuffer.Num();
-	if (IndexCount == 0)
+	return ConvertVertexBuffer(VertexBuffer, Builder);
+}
+
+FGLTFJsonAccessor* FGLTFIndexBufferConverterRaw::Convert(const FGLTFIndexArray* IndexBuffer, FString MeshName)
+{
+	const uint32 MaxVertexIndex = IndexBuffer->Num() - 1;
+	// NOTE: Even if a maximum value (i.e. UINT8_MAX or UINT16_MAX) would fit inside the type (i.e. uint8 or uint16), it's not allowed per gltf spec
+	// since the maximum values trigger primitive restart in some graphics APIs and would require client implementations to rebuild the index buffer.
+	if (MaxVertexIndex < UINT8_MAX) return ConvertIndexBuffer<uint8>(*IndexBuffer, Builder);
+	if (MaxVertexIndex < UINT16_MAX) return ConvertIndexBuffer<uint16>(*IndexBuffer, Builder);
+	if (MaxVertexIndex == UINT32_MAX)
+	{
+		Builder.LogWarning(
+			FString::Printf(TEXT("Maximum value (%d) used in indices accessor for mesh %s, this may not be supported on some glTF viewers"),
+				UINT32_MAX,
+				*MeshName
+			));
+	}
+	return ConvertIndexBuffer<uint32>(*IndexBuffer, Builder);
+}
+
+
+FGLTFJsonAccessor* FGLTFNormalBufferConverterRaw::Convert(const FGLTFNormalArray* NormalsSource)
+{
+	if (NormalsSource->Num() == 0)
 	{
 		return nullptr;
 	}
 
-	TArray<IndexType> Indices;
-	Indices.AddUninitialized(IndexCount);
+	FGLTFJsonAccessor* JsonAccessor = Builder.AddAccessor();
+	JsonAccessor->BufferView = ConvertBufferView(NormalsSource);;
+	JsonAccessor->ComponentType = EGLTFJsonComponentType::Float;
+	JsonAccessor->Count = NormalsSource->Num();
+	JsonAccessor->Type = EGLTFJsonAccessorType::Vec3;
 
-	for (uint32 Index = 0; Index < IndexCount; ++Index)
+	return JsonAccessor;
+}
+
+FGLTFJsonBufferView* FGLTFNormalBufferConverterRaw::ConvertBufferView(const FGLTFNormalArray* NormalsSource)
+{
+	const uint32 VertexCount = NormalsSource->Num();
+
+	TArray<FGLTFVector3> Normals;
+	Normals.AddUninitialized(VertexCount);
+
+	for (uint32 VertexIndex = 0; VertexIndex < VertexCount; ++VertexIndex)
 	{
-		Indices[Index] = static_cast<IndexType>(IndexBuffer[Index]);
+		Normals[VertexIndex] = FGLTFCoreUtilities::ConvertNormal((*NormalsSource)[VertexIndex]);
+	}
+
+	return Builder.AddBufferView(Normals, EGLTFJsonBufferTarget::ArrayBuffer);
+}
+
+FGLTFJsonAccessor* FGLTFUVBufferConverterRaw::Convert(const FGLTFUVArray* UVSource)
+{
+	if (UVSource->Num() == 0)
+	{
+		return nullptr;
 	}
 
 	FGLTFJsonAccessor* JsonAccessor = Builder.AddAccessor();
-	JsonAccessor->BufferView = Builder.AddBufferView(Indices, EGLTFJsonBufferTarget::ElementArrayBuffer, sizeof(IndexType));
-	JsonAccessor->ComponentType = FGLTFCoreUtilities::GetComponentType<IndexType>();
-	JsonAccessor->Count = IndexCount;
-	JsonAccessor->Type = EGLTFJsonAccessorType::Scalar;
+	JsonAccessor->BufferView = ConvertBufferView(UVSource);;
+	JsonAccessor->ComponentType = EGLTFJsonComponentType::Float;
+	JsonAccessor->Count = UVSource->Num();
+	JsonAccessor->Type = EGLTFJsonAccessorType::Vec2;
 
 	return JsonAccessor;
 }
+
+FGLTFJsonBufferView* FGLTFUVBufferConverterRaw::ConvertBufferView(const FGLTFUVArray* UVSource)
+{
+	const uint32 VertexCount = UVSource->Num();
+
+	TArray<FGLTFVector2> UVs;
+	UVs.AddUninitialized(VertexCount);
+
+	for (uint32 VertexIndex = 0; VertexIndex < VertexCount; ++VertexIndex)
+	{
+		UVs[VertexIndex] = FGLTFCoreUtilities::ConvertUV((*UVSource)[VertexIndex]);
+	}
+
+	return Builder.AddBufferView(UVs, EGLTFJsonBufferTarget::ArrayBuffer);
+}
\ No newline at end of file
diff --git a/Engine/Plugins/Enterprise/GLTFExporter/Source/GLTFExporter/Private/Converters/GLTFMeshConverters.cpp b/Engine/Plugins/Enterprise/GLTFExporter/Source/GLTFExporter/Private/Converters/GLTFMeshConverters.cpp
--- a/Engine/Plugins/Enterprise/GLTFExporter/Source/GLTFExporter/Private/Converters/GLTFMeshConverters.cpp
+++ b/Engine/Plugins/Enterprise/GLTFExporter/Source/GLTFExporter/Private/Converters/GLTFMeshConverters.cpp
@@ -10,6 +10,8 @@
 #include "Engine/StaticMesh.h"
 #include "Components/StaticMeshComponent.h"
 #include "Components/SkeletalMeshComponent.h"
+#include "Components/SplineMeshComponent.h"
+#include "LandscapeComponent.h"
 
 void FGLTFStaticMeshConverter::Sanitize(const UStaticMesh*& StaticMesh, const UStaticMeshComponent*& StaticMeshComponent, FGLTFMaterialArray& Materials, int32& LODIndex)
 {
@@ -95,3 +97,55 @@ FGLTFJsonMesh* FGLTFSkeletalMeshConverter::Convert(const USkeletalMesh* Skeletal
 	Builder.ScheduleSlowTask<FGLTFDelayedSkeletalMeshTask>(Builder, MeshSectionConverter, SkeletalMesh, SkeletalMeshComponent, Materials, LODIndex, JsonMesh);
 	return JsonMesh;
 }
+
+
+void FGLTFSplineMeshConverter::Sanitize(const UStaticMesh*& StaticMesh, const USplineMeshComponent*& SplineMeshComponent, FGLTFMaterialArray& Materials, int32& LODIndex)
+{
+	//SplineMeshComponent is used for calculating the SplineTransform, hence we cannot null it out. (Contrary to FGLTFStaticMeshConverter::Sanitize and FGLTFSkeletalMeshConverter::Sanitize)
+
+	FGLTFMeshUtilities::ResolveMaterials(Materials, SplineMeshComponent, StaticMesh);
+	Builder.ResolveProxies(Materials);
+}
+
+FGLTFJsonMesh* FGLTFSplineMeshConverter::Convert(const UStaticMesh* StaticMesh, const USplineMeshComponent* SplineMeshComponent, FGLTFMaterialArray Materials, int32 LODIndex)
+{
+#if !WITH_EDITOR
+	if (!StaticMesh->bAllowCPUAccess)
+	{
+		Builder.LogSuggestion(FString::Printf(
+			TEXT("Export of mesh %s can in runtime be speed-up by checking 'Allow CPU Access' in asset settings"),
+			*StaticMesh->GetName()));
+	}
+#endif
+
+	const int32 MaterialCount = FGLTFMeshUtilities::GetMaterials(StaticMesh).Num();
+	FGLTFJsonMesh* JsonMesh = Builder.AddMesh();
+	JsonMesh->Primitives.AddDefaulted(MaterialCount);
+
+	if (StaticMesh && SplineMeshComponent)
+	{
+		Builder.ScheduleSlowTask<FGLTFDelayedSplineMeshTask>(Builder, /*PositionBuffer,*/ MeshSectionConverter, *StaticMesh, *SplineMeshComponent, Materials, LODIndex, JsonMesh);
+	}
+	
+	return JsonMesh;
+}
+
+void FGLTFLandscapeMeshConverter::Sanitize(const ULandscapeComponent*& LandscapeComponent, const UMaterialInterface*& LandscapeMaterial)
+{
+	FGLTFMeshUtilities::ResolveMaterials(LandscapeMaterial, LandscapeComponent);
+	UGLTFMaterialExportOptions::ResolveProxy(LandscapeMaterial);
+}
+
+FGLTFJsonMesh* FGLTFLandscapeMeshConverter::Convert(const ULandscapeComponent* LandscapeComponent, const UMaterialInterface* LandscapeMaterial)
+{
+	const int32 MaterialCount = 1; //Sanitize will resolve it to either the LandscapeMaterial or the DefaultMaterial
+	FGLTFJsonMesh* JsonMesh = Builder.AddMesh();
+	JsonMesh->Primitives.AddDefaulted(MaterialCount);
+
+	if (LandscapeComponent && LandscapeMaterial)
+	{
+		Builder.ScheduleSlowTask<FGLTFDelayedLandscapeTask>(Builder, *LandscapeComponent, JsonMesh, *LandscapeMaterial);
+	}
+	
+	return JsonMesh;
+}
\ No newline at end of file
diff --git a/Engine/Plugins/Enterprise/GLTFExporter/Source/GLTFExporter/Private/Converters/GLTFMeshUtilities.cpp b/Engine/Plugins/Enterprise/GLTFExporter/Source/GLTFExporter/Private/Converters/GLTFMeshUtilities.cpp
--- a/Engine/Plugins/Enterprise/GLTFExporter/Source/GLTFExporter/Private/Converters/GLTFMeshUtilities.cpp
+++ b/Engine/Plugins/Enterprise/GLTFExporter/Source/GLTFExporter/Private/Converters/GLTFMeshUtilities.cpp
@@ -10,6 +10,7 @@
 #include "Rendering/SkeletalMeshRenderData.h"
 #include "StaticMeshCompiler.h"
 #include "SkinnedAssetCompiler.h"
+#include "LandscapeComponent.h"
 
 void FGLTFMeshUtilities::FullyLoad(const UStaticMesh* InStaticMesh)
 {
@@ -96,6 +97,16 @@ void FGLTFMeshUtilities::ResolveMaterials(TArray<const UMaterialInterface*>& Mat
 	ResolveMaterials(Materials, GetDefaultMaterial());
 }
 
+void FGLTFMeshUtilities::ResolveMaterials(const UMaterialInterface*& LandscapeMaterial, const ULandscapeComponent* LandscapeComponent)
+{
+	LandscapeMaterial = LandscapeComponent->GetLandscapeMaterial();
+
+	if (LandscapeMaterial == nullptr)
+	{
+		LandscapeMaterial = GetDefaultMaterial();
+	}
+}
+
 template <typename MaterialType>
 void FGLTFMeshUtilities::ResolveMaterials(TArray<const UMaterialInterface*>& Materials, const TArray<MaterialType>& Defaults)
 {
diff --git a/Engine/Plugins/Enterprise/GLTFExporter/Source/GLTFExporter/Private/Converters/GLTFMeshUtilities.h b/Engine/Plugins/Enterprise/GLTFExporter/Source/GLTFExporter/Private/Converters/GLTFMeshUtilities.h
--- a/Engine/Plugins/Enterprise/GLTFExporter/Source/GLTFExporter/Private/Converters/GLTFMeshUtilities.h
+++ b/Engine/Plugins/Enterprise/GLTFExporter/Source/GLTFExporter/Private/Converters/GLTFMeshUtilities.h
@@ -10,6 +10,7 @@ class UStaticMeshComponent;
 class USkeletalMesh;
 class USkeletalMeshComponent;
 class FSkeletalMeshLODRenderData;
+class ULandscapeComponent;
 struct FStaticMaterial;
 struct FSkeletalMaterial;
 struct FStaticMeshLODResources;
@@ -30,6 +31,7 @@ struct FGLTFMeshUtilities
 
 	static void ResolveMaterials(TArray<const UMaterialInterface*>& Materials, const UStaticMeshComponent* StaticMeshComponent, const UStaticMesh* StaticMesh);
 	static void ResolveMaterials(TArray<const UMaterialInterface*>& Materials, const USkeletalMeshComponent* SkeletalMeshComponent, const USkeletalMesh* SkeletalMesh);
+	static void ResolveMaterials(const UMaterialInterface*& LandscapeMaterial, const ULandscapeComponent* LandscapeComponent);
 
 	template <typename MaterialType>
 	static void ResolveMaterials(TArray<const UMaterialInterface*>& Materials, const TArray<MaterialType>& Defaults);
diff --git a/Engine/Plugins/Enterprise/GLTFExporter/Source/GLTFExporter/Private/Converters/GLTFNodeConverters.cpp b/Engine/Plugins/Enterprise/GLTFExporter/Source/GLTFExporter/Private/Converters/GLTFNodeConverters.cpp
--- a/Engine/Plugins/Enterprise/GLTFExporter/Source/GLTFExporter/Private/Converters/GLTFNodeConverters.cpp
+++ b/Engine/Plugins/Enterprise/GLTFExporter/Source/GLTFExporter/Private/Converters/GLTFNodeConverters.cpp
@@ -2,7 +2,6 @@
 
 #include "Converters/GLTFNodeConverters.h"
 #include "Converters/GLTFNameUtilities.h"
-#include "Converters/GLTFBoneUtilities.h"
 #include "Builders/GLTFContainerBuilder.h"
 #include "Utilities/GLTFCoreUtilities.h"
 #include "LevelSequenceActor.h"
@@ -16,17 +15,28 @@
 #include "Engine/SkeletalMeshSocket.h"
 #include "GameFramework/Pawn.h"
 
+#include "LandscapeComponent.h"
+#include "GLTFMeshUtilities.h"
+#include "Components/InstancedStaticMeshComponent.h"
+#include "Components/SplineMeshComponent.h"
+#include "LevelInstance/LevelInstanceComponent.h"
+#include "LevelInstance/LevelInstanceActor.h"
+#include "LevelInstance/LevelInstanceEditorInstanceActor.h"
+#include "Engine/Level.h"
+
 FGLTFJsonNode* FGLTFActorConverter::Convert(const AActor* Actor)
 {
-	if (Actor->bIsEditorOnlyActor)
+	if (!Actor || !Actor->IsValidLowLevel() || Actor->IsEditorOnly() || !Builder.IsSelectedActor(Actor))
 	{
 		return nullptr;
 	}
-
-	if (!Builder.IsSelectedActor(Actor))
+#if WITH_EDITOR
+	if (const ALevelInstanceEditorInstanceActor* LevelInstanceEditorInstanceActor = Cast<ALevelInstanceEditorInstanceActor>(Actor))
 	{
+		//Editor only class, reporting as IsEditorOnly==false
 		return nullptr;
 	}
+#endif
 
 	const USceneComponent* RootComponent = Actor->GetRootComponent();
 	FGLTFJsonNode* RootNode = Builder.AddUniqueNode(RootComponent);
@@ -45,7 +55,6 @@ FGLTFJsonNode* FGLTFActorConverter::Convert(const AActor* Actor)
 		// TODO: add support for exporting brush geometry?
 
 		// TODO: to reduce number of nodes, only export components that are of interest
-
 		for (const UActorComponent* Component : Actor->GetComponents())
 		{
 			const USceneComponent* SceneComponent = Cast<USceneComponent>(Component);
@@ -61,7 +70,7 @@ FGLTFJsonNode* FGLTFActorConverter::Convert(const AActor* Actor)
 
 FGLTFJsonNode* FGLTFComponentConverter::Convert(const USceneComponent* SceneComponent)
 {
-	if (SceneComponent->IsEditorOnly())
+	if (!SceneComponent || SceneComponent->IsEditorOnly())
 	{
 		return nullptr;
 	}
@@ -73,15 +82,23 @@ FGLTFJsonNode* FGLTFComponentConverter::Convert(const USceneComponent* SceneComp
 		return nullptr;
 	}
 
-	if (!Builder.IsSelectedActor(Owner))
+#if WITH_EDITOR
+	if (const ALevelInstanceEditorInstanceActor* LevelInstanceEditorInstanceActor = Cast<ALevelInstanceEditorInstanceActor>(Owner))
 	{
+		//Editor only class, reporting as IsEditorOnly==false
 		return nullptr;
 	}
+#endif
 
+	if (!Builder.IsSelectedActor(Owner))
+	{
+		return nullptr;
+	}
+	
 	const bool bIsRootComponent = Owner->GetRootComponent() == SceneComponent;
 	const bool bIsRootNode = bIsRootComponent && Builder.IsRootActor(Owner);
-
 	const USceneComponent* ParentComponent = !bIsRootNode ? SceneComponent->GetAttachParent() : nullptr;
+
 	const FName SocketName = SceneComponent->GetAttachSocketName();
 	FGLTFJsonNode* ParentNode = Builder.AddUniqueNode(ParentComponent, SocketName);
 
@@ -130,8 +147,90 @@ void FGLTFComponentConverter::ConvertComponentSpecialization(const USceneCompone
 	{
 		return;
 	}
+	
+    if (const UInstancedStaticMeshComponent* InstancedMeshComp = Cast<UInstancedStaticMeshComponent>(SceneComponent))
+	{
+		const UStaticMesh* StaticMesh = InstancedMeshComp->GetStaticMesh();
+
+		FGLTFJsonMesh* Mesh = Builder.AddUniqueMesh(StaticMesh);
+
+		const int32 NumInstances = InstancedMeshComp->GetInstanceCount();
+		for (int32 InstanceIndex = 0; InstanceIndex < NumInstances; ++InstanceIndex)
+		{
+			FTransform RelativeTransform;
+			if (ensure(InstancedMeshComp->GetInstanceTransform(InstanceIndex, RelativeTransform, /*bWorldSpace=*/false)))
+			{
+				FGLTFJsonNode* InstanceNode = Builder.AddNode();
+				InstanceNode->Mesh = Mesh;
+				InstanceNode->Name = FString::Printf(TEXT("%s:%d"), *(InstancedMeshComp->GetName()), InstanceIndex);
+				const FTransform3f Transform = FTransform3f(RelativeTransform);
+
+				InstanceNode->Translation = FGLTFCoreUtilities::ConvertPosition(Transform.GetTranslation(), Builder.ExportOptions->ExportUniformScale);
+				InstanceNode->Rotation = FGLTFCoreUtilities::ConvertRotation(Transform.GetRotation());
+				InstanceNode->Scale = FGLTFCoreUtilities::ConvertScale(Transform.GetScale3D());
+				Node->Children.Add(InstanceNode);
+			}
+		}
+	}
+	else if (const USplineMeshComponent* SplineMeshComponent = Cast<USplineMeshComponent>(SceneComponent))
+	{
+		FGLTFJsonNode* SplineNode = Builder.AddNode();
+		SplineNode->Name = SplineMeshComponent->GetName();
+		SplineNode->Mesh = Builder.AddUniqueMesh(SplineMeshComponent);
+
+		const FTransform3f Transform;
+		SplineNode->Translation = FGLTFCoreUtilities::ConvertPosition(Transform.GetTranslation(), Builder.ExportOptions->ExportUniformScale);
+		SplineNode->Rotation = FGLTFCoreUtilities::ConvertRotation(Transform.GetRotation());
+		SplineNode->Scale = FGLTFCoreUtilities::ConvertScale(Transform.GetScale3D());
+
+		Node->Children.Add(SplineNode);
+	}
+	else if (const ULandscapeComponent* LandscapeComponent = Cast<ULandscapeComponent>(SceneComponent))
+	{
+		FGLTFJsonNode* LandscapeNode = Builder.AddNode();
+		LandscapeNode->Name = LandscapeComponent->GetName();
+		LandscapeNode->Mesh = Builder.AddUniqueMesh(LandscapeComponent);
+		
+		const FTransform3f Transform;
+		LandscapeNode->Translation = FGLTFCoreUtilities::ConvertPosition(Transform.GetTranslation(), Builder.ExportOptions->ExportUniformScale);
+		LandscapeNode->Rotation = FGLTFCoreUtilities::ConvertRotation(Transform.GetRotation());
+		LandscapeNode->Scale = FGLTFCoreUtilities::ConvertScale(Transform.GetScale3D());
+		
+		Node->Children.Add(LandscapeNode);
+	}
+	else if (const ULevelInstanceComponent* LevelInstanceComponent = Cast<ULevelInstanceComponent>(SceneComponent))
+	{
+		const ALevelInstance* LevelInstance = Cast<ALevelInstance>(Owner);
+		
+		TArray<TObjectPtr<AActor>> Actors = LevelInstance->GetLoadedLevel()->Actors;
+
+		for (const AActor* LevelActor : Actors)
+		{
+			FGLTFJsonNode* LevelActorNode = Builder.AddUniqueNode(LevelActor);
 
-	if (const UStaticMeshComponent* StaticMeshComponent = Cast<UStaticMeshComponent>(SceneComponent))
+			if (LevelActorNode == nullptr)
+			{
+				continue;
+			}
+			const USceneComponent* LevelActorComponent = LevelActor->GetRootComponent();
+			const USceneComponent* ParentComponent = SceneComponent;
+
+			const FName SocketName = LevelActorComponent->GetAttachSocketName();
+			const FTransform3f Transform = FTransform3f(LevelActorComponent->GetComponentTransform());
+			const FTransform3f ParentTransform = ParentComponent != nullptr ? FTransform3f(ParentComponent->GetSocketTransform(SocketName)) : FTransform3f::Identity;
+			const FTransform3f RelativeTransform = Transform.GetRelativeTransform(ParentTransform);
+
+			LevelActorNode->Translation = FGLTFCoreUtilities::ConvertPosition(RelativeTransform.GetTranslation(), Builder.ExportOptions->ExportUniformScale);
+			LevelActorNode->Rotation = FGLTFCoreUtilities::ConvertRotation(RelativeTransform.GetRotation());
+			LevelActorNode->Scale = FGLTFCoreUtilities::ConvertScale(RelativeTransform.GetScale3D());
+
+			if (LevelActorNode != nullptr)
+			{
+				Node->Children.Add(LevelActorNode);
+			}
+		}
+	}
+	else if (const UStaticMeshComponent* StaticMeshComponent = Cast<UStaticMeshComponent>(SceneComponent))
 	{
 		Node->Mesh = Builder.AddUniqueMesh(StaticMeshComponent);
 	}
diff --git a/Engine/Plugins/Enterprise/GLTFExporter/Source/GLTFExporter/Private/Converters/GLTFSceneConverters.cpp b/Engine/Plugins/Enterprise/GLTFExporter/Source/GLTFExporter/Private/Converters/GLTFSceneConverters.cpp
--- a/Engine/Plugins/Enterprise/GLTFExporter/Source/GLTFExporter/Private/Converters/GLTFSceneConverters.cpp
+++ b/Engine/Plugins/Enterprise/GLTFExporter/Source/GLTFExporter/Private/Converters/GLTFSceneConverters.cpp
@@ -18,6 +18,12 @@ FGLTFJsonScene* FGLTFSceneConverter::Convert(const UWorld* World)
 	{
 		for (const ULevel* Level : Levels)
 		{
+			if (Level->IsInstancedLevel() && !Level->IsCurrentLevel())
+			{
+				//LevelInstances are handled in FGLTFComponentConverter::ConvertComponentSpecialization in the ULevelInstanceComponent handling branch.
+				continue;
+			}
+
 			if (Level == nullptr)
 			{
 				continue;
diff --git a/Engine/Plugins/Enterprise/GLTFExporter/Source/GLTFExporter/Private/Tasks/GLTFDelayedMeshTasks.cpp b/Engine/Plugins/Enterprise/GLTFExporter/Source/GLTFExporter/Private/Tasks/GLTFDelayedMeshTasks.cpp
--- a/Engine/Plugins/Enterprise/GLTFExporter/Source/GLTFExporter/Private/Tasks/GLTFDelayedMeshTasks.cpp
+++ b/Engine/Plugins/Enterprise/GLTFExporter/Source/GLTFExporter/Private/Tasks/GLTFDelayedMeshTasks.cpp
@@ -13,6 +13,16 @@
 #include "StaticMeshResources.h"
 #include "Rendering/SkeletalMeshRenderData.h"
 
+#include "Converters/GLTFMaterialUtilities.h"
+#include "Rendering/PositionVertexBuffer.h"
+#include "Components/SplineMeshComponent.h"
+#include "LandscapeProxy.h"
+#include "LandscapeComponent.h"
+#include "LandscapeLayerInfoObject.h"
+#include "Converters/GLTFNormalArray.h"
+#include "Converters/GLTFUVArray.h"
+#include "Utilities/GLTFLandscapeComponentDataInterface.h"
+
 namespace
 {
 	template <typename VectorType>
@@ -178,9 +188,9 @@ void FGLTFDelayedStaticMeshTask::Process()
 		{
 			Builder.LogError(
 				FString::Printf(TEXT("Failed to export vertex positions related to material slot %d (%s) in static mesh %s"),
-				MaterialIndex,
-				*MaterialSlots[MaterialIndex].MaterialSlotName.ToString(),
-				*ConvertedSection->ToString()
+					MaterialIndex,
+					*MaterialSlots[MaterialIndex].MaterialSlotName.ToString(),
+					*ConvertedSection->ToString()
 				));
 		}
 
@@ -203,7 +213,7 @@ void FGLTFDelayedStaticMeshTask::Process()
 		}
 
 		const UMaterialInterface* Material = Materials[MaterialIndex];
-		JsonPrimitive.Material =  Builder.AddUniqueMaterial(Material, MeshData, SectionIndices);
+		JsonPrimitive.Material = Builder.AddUniqueMaterial(Material, MeshData, SectionIndices);
 	}
 }
 
@@ -280,9 +290,9 @@ void FGLTFDelayedSkeletalMeshTask::Process()
 		{
 			Builder.LogError(
 				FString::Printf(TEXT("Failed to export vertex positions related to material slot %d (%s) in skeletal mesh %s"),
-				MaterialIndex,
-				*MaterialSlots[MaterialIndex].MaterialSlotName.ToString(),
-				*ConvertedSection->ToString()
+					MaterialIndex,
+					*MaterialSlots[MaterialIndex].MaterialSlotName.ToString(),
+					*ConvertedSection->ToString()
 				));
 		}
 
@@ -319,6 +329,250 @@ void FGLTFDelayedSkeletalMeshTask::Process()
 		}
 
 		const UMaterialInterface* Material = Materials[MaterialIndex];
-		JsonPrimitive.Material =  Builder.AddUniqueMaterial(Material, MeshData, SectionIndices);
+		JsonPrimitive.Material = Builder.AddUniqueMaterial(Material, MeshData, SectionIndices);
+	}
+}
+
+
+FString FGLTFDelayedSplineMeshTask::GetName()
+{
+	return FGLTFNameUtilities::GetName(&SplineMeshComponent);
+}
+
+void FGLTFDelayedSplineMeshTask::Process()
+{
+	FGLTFMeshUtilities::FullyLoad(&StaticMesh);
+	JsonMesh->Name = GetName();
+
+	const FStaticMeshLODResources& RenderData = FGLTFMeshUtilities::GetRenderData(&StaticMesh, LODIndex);
+	const FPositionVertexBuffer& PositionBuffer = RenderData.VertexBuffers.PositionVertexBuffer;
+	const FStaticMeshVertexBuffer& VertexBuffer = RenderData.VertexBuffers.StaticMeshVertexBuffer;
+	const FColorVertexBuffer* ColorBuffer = &RenderData.VertexBuffers.ColorVertexBuffer; // TODO: add support for overriding color buffer by component
+
+	if (Builder.ExportOptions->bExportVertexColors && HasVertexColors(ColorBuffer))
+	{
+		Builder.LogSuggestion(FString::Printf(
+			TEXT("Vertex colors in mesh %s will act as a multiplier for base color in glTF, regardless of material, which may produce undesirable results"),
+			*StaticMesh.GetName()));
+	}
+	else
+	{
+		ColorBuffer = nullptr;
+	}
+
+	if (SplineMeshComponent.LODData.IsValidIndex(LODIndex))
+	{
+		const FStaticMeshComponentLODInfo& LODInfo = SplineMeshComponent.LODData[LODIndex];
+		ColorBuffer = LODInfo.OverrideVertexColors != nullptr ? LODInfo.OverrideVertexColors : ColorBuffer;
+	}
+
+	const FGLTFMeshData* MeshData = Builder.ExportOptions->BakeMaterialInputs == EGLTFMaterialBakeMode::UseMeshData ?
+	Builder.AddUniqueMeshData(&StaticMesh, &SplineMeshComponent, LODIndex) : nullptr;
+
+	#if WITH_EDITOR
+	if (MeshData != nullptr)
+	{
+		if (MeshData->Description.IsEmpty())
+		{
+			// TODO: report warning in case the mesh actually has data, which means we failed to extract a mesh description.
+			MeshData = nullptr;
+		}
+		else if (MeshData->BakeUsingTexCoord < 0)
+		{
+			// TODO: report warning (about missing texture coordinate for baking with mesh data).
+			MeshData = nullptr;
+		}
+	}
+	#endif
+
+	ValidateVertexBuffer(Builder, &VertexBuffer, *StaticMesh.GetName());
+
+	const TArray<FStaticMaterial>& MaterialSlots = FGLTFMeshUtilities::GetMaterials(&StaticMesh);
+	const int32 MaterialCount = MaterialSlots.Num();
+
+	for (int32 MaterialIndex = 0; MaterialIndex < MaterialCount; ++MaterialIndex)
+	{
+		const FGLTFIndexArray SectionIndices = FGLTFMeshUtilities::GetSectionIndices(RenderData, MaterialIndex);
+		const FGLTFMeshSection* ConvertedSection = MeshSectionConverter.GetOrAdd(&StaticMesh, LODIndex, SectionIndices);
+
+		FGLTFJsonPrimitive& JsonPrimitive = JsonMesh->Primitives[MaterialIndex];
+		JsonPrimitive.Indices = Builder.AddUniqueIndexAccessor(ConvertedSection);
+
+		FPositionVertexBuffer* TransformedPositionBuffer = nullptr;
+		
+		{//fix for Splines:
+			TransformedPositionBuffer = new FPositionVertexBuffer();
+			TransformedPositionBuffer->Init(PositionBuffer.GetNumVertices(), true);
+
+			const uint32 VertexCount = PositionBuffer.GetNumVertices();
+			const uint32 Stride = PositionBuffer.GetStride();
+
+			const TUniquePtr<IGLTFBufferAdapter> SourceBuffer = IGLTFBufferAdapter::GetPositions(&PositionBuffer);
+			const uint8* SourceData = SourceBuffer->GetData();
+
+			for (uint32 VertexIndex = 0; VertexIndex < VertexCount; ++VertexIndex)
+			{
+				FVector3f Position = *reinterpret_cast<const FVector3f*>(SourceData + Stride * VertexIndex);
+
+				const FTransform3f SliceTransform = FTransform3f(SplineMeshComponent.CalcSliceTransform(USplineMeshComponent::GetAxisValueRef(Position, SplineMeshComponent.ForwardAxis)));
+				USplineMeshComponent::GetAxisValueRef(Position, SplineMeshComponent.ForwardAxis) = 0;
+				Position = SliceTransform.TransformPosition(Position);
+
+				TransformedPositionBuffer->VertexPosition(VertexIndex) = FVector3f(Position.X, Position.Y, Position.Z);
+			}
+		}
+
+		JsonPrimitive.Attributes.Position = Builder.AddUniquePositionAccessor(ConvertedSection, TransformedPositionBuffer != nullptr ? TransformedPositionBuffer : &PositionBuffer);
+		if (JsonPrimitive.Attributes.Position == nullptr)
+		{
+			Builder.LogError(
+				FString::Printf(TEXT("Failed to export vertex positions related to material slot %d (%s) in static mesh %s"),
+					MaterialIndex,
+					*MaterialSlots[MaterialIndex].MaterialSlotName.ToString(),
+					*ConvertedSection->ToString()
+				));
+		}
+
+		if (ColorBuffer != nullptr)
+		{
+			JsonPrimitive.Attributes.Color0 = Builder.AddUniqueColorAccessor(ConvertedSection, ColorBuffer);
+		}
+
+		// TODO: report warning if both Mesh Quantization (export options) and Use High Precision Tangent Basis (vertex buffer) are disabled
+		JsonPrimitive.Attributes.Normal = Builder.AddUniqueNormalAccessor(ConvertedSection, &VertexBuffer);
+		JsonPrimitive.Attributes.Tangent = Builder.AddUniqueTangentAccessor(ConvertedSection, &VertexBuffer);
+
+		const uint32 UVCount = VertexBuffer.GetNumTexCoords();
+		// TODO: report warning or option to limit UV channels since most viewers don't support more than 2?
+		JsonPrimitive.Attributes.TexCoords.AddUninitialized(UVCount);
+
+		for (uint32 UVIndex = 0; UVIndex < UVCount; ++UVIndex)
+		{
+			JsonPrimitive.Attributes.TexCoords[UVIndex] = Builder.AddUniqueUVAccessor(ConvertedSection, &VertexBuffer, UVIndex);
+		}
+
+		const UMaterialInterface* Material = Materials[MaterialIndex];
+		JsonPrimitive.Material = Builder.AddUniqueMaterial(Material, MeshData, SectionIndices);
 	}
 }
+
+
+FGLTFDelayedLandscapeTask::FGLTFDelayedLandscapeTask(FGLTFConvertBuilder& Builder, const ULandscapeComponent& LandscapeComponent, FGLTFJsonMesh* JsonMesh, const UMaterialInterface& LandscapeMaterial)
+	: FGLTFDelayedTask(EGLTFTaskPriority::Mesh)
+	, Builder(Builder)
+	, LandscapeComponent(LandscapeComponent)
+	, JsonMesh(JsonMesh)
+	, LandscapeMaterial(LandscapeMaterial)
+{
+}
+
+FString FGLTFDelayedLandscapeTask::GetName()
+{
+	return LandscapeComponent.GetName();
+}
+
+void FGLTFDelayedLandscapeTask::Process()
+{
+	const ALandscapeProxy* Landscape = Cast<ALandscapeProxy>(LandscapeComponent.GetOwner());
+	JsonMesh->Name = LandscapeComponent.GetName();
+
+	int32 MinX = MAX_int32, MinY = MAX_int32;
+	int32 MaxX = MIN_int32, MaxY = MIN_int32;
+
+	// Create and fill in the vertex position data source.
+	int32 ExportLOD = 0;
+#if WITH_EDITOR
+	ExportLOD = Landscape->ExportLOD;
+#endif
+	const int32 ComponentSizeQuads = ((Landscape->ComponentSizeQuads + 1) >> ExportLOD) - 1;
+	const float ScaleFactor = (float)Landscape->ComponentSizeQuads / (float)ComponentSizeQuads;
+	const int32 VertexCount = FMath::Square(ComponentSizeQuads + 1);
+	const int32 TriangleCount = FMath::Square(ComponentSizeQuads) * 2;
+
+	FPositionVertexBuffer* PositionBuffer = new FPositionVertexBuffer();
+	PositionBuffer->Init(VertexCount, true);
+
+	FGLTFJsonPrimitive& JsonPrimitive = JsonMesh->Primitives[0];
+	TArray<uint8> VisibilityData;
+	VisibilityData.Empty(VertexCount);
+	VisibilityData.AddZeroed(VertexCount);
+
+	int OffsetX = Landscape->LandscapeSectionOffset.X;
+	int OffsetY = Landscape->LandscapeSectionOffset.Y;
+	
+	FGLTFLandscapeComponentDataInterface CDI(LandscapeComponent, ExportLOD);
+
+	TArray<uint8> CompVisData;
+	const TArray<FWeightmapLayerAllocationInfo>& ComponentWeightmapLayerAllocations = LandscapeComponent.GetWeightmapLayerAllocations();
+
+	for (int32 AllocIdx = 0; AllocIdx < ComponentWeightmapLayerAllocations.Num(); AllocIdx++)
+	{
+		const FWeightmapLayerAllocationInfo& AllocInfo = ComponentWeightmapLayerAllocations[AllocIdx];
+		//Landscape Visibility Layer is named: __LANDSCAPE_VISIBILITY__
+		//based on: Engine/Source/Runtime/Landscape/Private/Materials/MaterialExpressionLandscapeVisibilityMask.cpp
+		//		FName UMaterialExpressionLandscapeVisibilityMask::ParameterName = FName("__LANDSCAPE_VISIBILITY__");
+		FString LayerName = AllocInfo.LayerInfo->LayerName.ToString();
+		if (LayerName == TEXT("__LANDSCAPE_VISIBILITY__"))
+		{
+			CDI.GetWeightmapTextureData(AllocInfo.LayerInfo, CompVisData);
+		}
+	}
+	CompVisData.AddZeroed((ComponentSizeQuads + 1) * (ComponentSizeQuads + 1));
+
+	if (CompVisData.Num() > 0)
+	{
+		for (int32 i = 0; i < VertexCount; ++i)
+		{
+			VisibilityData[i] = CompVisData[CDI.VertexIndexToTexel(i)];
+		}
+	}
+
+	FGLTFNormalArray* Normals = new FGLTFNormalArray();
+	Normals->SetNumZeroed(VertexCount);
+	FGLTFUVArray* UVs = new FGLTFUVArray(); //UVIndex==0
+	UVs->SetNumZeroed(VertexCount);
+
+	for (int32 VertexIndex = 0; VertexIndex < VertexCount; VertexIndex++)
+	{
+		int32 VertX, VertY;
+		CDI.VertexIndexToXY(VertexIndex, VertX, VertY);
+
+		FVector3f Position;
+		FVector3f Normal;
+		FVector2f UV;
+		CDI.GetPositionNormalUV(VertX, VertY, Position, Normal, UV);
+
+		PositionBuffer->VertexPosition(VertexIndex) = Position;
+		(*Normals)[VertexIndex] = Normal;
+		(*UVs)[VertexIndex] = UV;
+	}
+
+	const int32 VisThreshold = 170;
+	
+	FGLTFIndexArray* LandscapeIndices = new FGLTFIndexArray();
+	LandscapeIndices->Reserve(FMath::Square(ComponentSizeQuads) * 2 * 3);
+
+	for (int32 Y = 0; Y < ComponentSizeQuads; Y++)
+	{
+		for (int32 X = 0; X < ComponentSizeQuads; X++)
+		{
+			if (VisibilityData[Y * (ComponentSizeQuads + 1) + X] < VisThreshold)
+			{
+				LandscapeIndices->Push((X + 0) + (Y + 0) * (ComponentSizeQuads + 1));
+				LandscapeIndices->Push((X + 1) + (Y + 1) * (ComponentSizeQuads + 1));
+				LandscapeIndices->Push((X + 1) + (Y + 0) * (ComponentSizeQuads + 1));
+
+				LandscapeIndices->Push((X + 0) + (Y + 0) * (ComponentSizeQuads + 1));
+				LandscapeIndices->Push((X + 0) + (Y + 1) * (ComponentSizeQuads + 1));
+				LandscapeIndices->Push((X + 1) + (Y + 1) * (ComponentSizeQuads + 1));
+			}
+		}
+	}
+
+	JsonPrimitive.Attributes.Position = Builder.AddUniquePositionAccessor(PositionBuffer);
+	JsonPrimitive.Attributes.Normal = Builder.AddUniqueNormalAccessor(Normals);
+	JsonPrimitive.Attributes.TexCoords.AddUninitialized(1);
+	JsonPrimitive.Attributes.TexCoords[0] = Builder.AddUniqueUVAccessor(UVs);
+	JsonPrimitive.Indices = Builder.AddUniqueIndexAccessor(LandscapeIndices, JsonMesh->Name);
+	JsonPrimitive.Material = Builder.AddUniqueMaterial(&LandscapeMaterial);
+}
\ No newline at end of file
diff --git a/Engine/Plugins/Enterprise/GLTFExporter/Source/GLTFExporter/Private/Tasks/GLTFDelayedMeshTasks.h b/Engine/Plugins/Enterprise/GLTFExporter/Source/GLTFExporter/Private/Tasks/GLTFDelayedMeshTasks.h
--- a/Engine/Plugins/Enterprise/GLTFExporter/Source/GLTFExporter/Private/Tasks/GLTFDelayedMeshTasks.h
+++ b/Engine/Plugins/Enterprise/GLTFExporter/Source/GLTFExporter/Private/Tasks/GLTFDelayedMeshTasks.h
@@ -74,3 +74,53 @@ private:
 	const int32 LODIndex;
 	FGLTFJsonMesh* JsonMesh;
 };
+
+
+class FGLTFDelayedSplineMeshTask : public FGLTFDelayedTask
+{
+public:
+
+	FGLTFDelayedSplineMeshTask(FGLTFConvertBuilder& Builder, FGLTFStaticMeshSectionConverter& MeshSectionConverter, const UStaticMesh& StaticMesh, const USplineMeshComponent& SplineMeshComponent, FGLTFMaterialArray Materials, int32 LODIndex, FGLTFJsonMesh* JsonMesh)
+		: FGLTFDelayedTask(EGLTFTaskPriority::Mesh)
+		, Builder(Builder)
+		, MeshSectionConverter(MeshSectionConverter)
+		, StaticMesh(StaticMesh)
+		, SplineMeshComponent(SplineMeshComponent)
+		, Materials(Materials)
+		, LODIndex(LODIndex)
+		, JsonMesh(JsonMesh)
+	{
+	}
+
+	virtual FString GetName() override;
+
+	virtual void Process() override;
+
+private:
+
+	FGLTFConvertBuilder& Builder;
+	FGLTFStaticMeshSectionConverter& MeshSectionConverter;
+	const UStaticMesh& StaticMesh;
+	const USplineMeshComponent& SplineMeshComponent;
+	const FGLTFMaterialArray Materials;
+	const int32 LODIndex;
+	FGLTFJsonMesh* JsonMesh;
+};
+
+class FGLTFDelayedLandscapeTask : public FGLTFDelayedTask
+{
+public:
+
+	FGLTFDelayedLandscapeTask(FGLTFConvertBuilder& Builder, const ULandscapeComponent& LandscapeComponent, FGLTFJsonMesh* JsonMesh, const UMaterialInterface& LandscapeMaterial);
+
+	virtual FString GetName() override;
+
+	virtual void Process() override;
+
+private:
+
+	FGLTFConvertBuilder& Builder;
+	const ULandscapeComponent& LandscapeComponent;
+	FGLTFJsonMesh* JsonMesh;
+	const UMaterialInterface& LandscapeMaterial;
+};
diff --git a/Engine/Plugins/Enterprise/GLTFExporter/Source/GLTFExporter/Private/Utilities/GLTFLandscapeComponentDataInterface.cpp b/Engine/Plugins/Enterprise/GLTFExporter/Source/GLTFExporter/Private/Utilities/GLTFLandscapeComponentDataInterface.cpp
new file mode 100644
--- /dev/null
+++ b/Engine/Plugins/Enterprise/GLTFExporter/Source/GLTFExporter/Private/Utilities/GLTFLandscapeComponentDataInterface.cpp
@@ -0,0 +1,254 @@
+// Copyright Epic Games, Inc. All Rights Reserved.
+
+#include "Utilities/GLTFLandscapeComponentDataInterface.h"
+
+#include "LandscapeComponent.h"
+#include "LandscapeDataAccess.h"
+#include "LandscapeLayerInfoObject.h"
+#include "TextureResource.h"
+
+FGLTFLandscapeComponentDataInterface::FGLTFLandscapeComponentDataInterface(const ULandscapeComponent& Component, int32 ExportLOD)
+	: Component(Component)
+	, MipLevel(ExportLOD)
+{
+
+	UTexture2D* HeightMapTexture = Component.GetHeightmap(true);
+#if WITH_EDITOR
+	if (HeightMapTexture)
+	{
+		if (MipLevel < HeightMapTexture->Source.GetNumMips())
+		{
+			HeightMapTexture->Source.GetMipData(HeightMipDataSource, MipLevel);
+			HeightMipData = (FColor*)HeightMipDataSource.GetData();
+		}
+	}
+
+	if (Component.XYOffsetmapTexture)
+	{
+		if (MipLevel < Component.XYOffsetmapTexture->Source.GetNumMips())
+		{
+			Component.XYOffsetmapTexture->Source.GetMipData(XYOffsetMipDataSource, MipLevel);
+			XYOffsetMipData = (FColor*)XYOffsetMipDataSource.GetData();
+		}
+	}
+#else
+	if (HeightMapTexture)
+	{
+		FTexturePlatformData* PlatformData = HeightMapTexture->GetPlatformData();
+		if (MipLevel < PlatformData->Mips.Num())
+		{
+			void* MipData = nullptr;
+			PlatformData->Mips[MipLevel].BulkData.GetCopy(&MipData, false);
+			HeightMipData = (FColor*)MipData;
+		}
+	}
+
+	if (Component.XYOffsetmapTexture)
+	{
+		FTexturePlatformData* PlatformData = Component.XYOffsetmapTexture->GetPlatformData();
+
+		if (MipLevel < PlatformData->Mips.Num())
+		{
+			void* MipData = nullptr;
+			PlatformData->Mips[MipLevel].BulkData.GetCopy(&MipData, false);
+			XYOffsetMipData = (FColor*)MipData;
+		}
+	}
+#endif
+	ComponentSubsectionSizeQuads = Component.SubsectionSizeQuads;
+
+	ComponentSizeQuads = Component.ComponentSizeQuads;
+	ComponentSizeVerts = (Component.ComponentSizeQuads + 1) >> MipLevel;
+	ComponentSubsectionSizeVerts = (ComponentSubsectionSizeQuads + 1) >> MipLevel;
+	ComponentNumSubsections = Component.NumSubsections;
+
+	if (HeightMapTexture)
+	{
+		HeightmapStride = HeightMapTexture->GetSizeX() >> MipLevel;
+		HeightmapComponentOffsetX = FMath::RoundToInt((HeightMapTexture->GetSizeX() >> MipLevel) * Component.HeightmapScaleBias.Z);
+		HeightmapComponentOffsetY = FMath::RoundToInt((HeightMapTexture->GetSizeY() >> MipLevel) * Component.HeightmapScaleBias.W);
+	}
+
+	HeightmapSubsectionOffset = (ComponentSubsectionSizeQuads + 1) >> MipLevel;
+}
+
+void FGLTFLandscapeComponentDataInterface::ComponentXYToSubsectionXY(int32 CompX, int32 CompY, int32& SubNumX, int32& SubNumY, int32& SubX, int32& SubY) const
+{
+	// We do the calculation as if we're looking for the previous vertex.
+	// This allows us to pick up the last shared vertex of every subsection correctly.
+	SubNumX = (CompX - 1) / (ComponentSubsectionSizeVerts - 1);
+	SubNumY = (CompY - 1) / (ComponentSubsectionSizeVerts - 1);
+	SubX = (CompX - 1) % (ComponentSubsectionSizeVerts - 1) + 1;
+	SubY = (CompY - 1) % (ComponentSubsectionSizeVerts - 1) + 1;
+
+	// If we're asking for the first vertex, the calculation above will lead
+	// to a negative SubNumX/Y, so we need to fix that case up.
+	if (SubNumX < 0)
+	{
+		SubNumX = 0;
+		SubX = 0;
+	}
+
+	if (SubNumY < 0)
+	{
+		SubNumY = 0;
+		SubY = 0;
+	}
+}
+
+int32 FGLTFLandscapeComponentDataInterface::TexelXYToIndex(int32 TexelX, int32 TexelY) const
+{
+	return TexelY * ComponentNumSubsections * ComponentSubsectionSizeVerts + TexelX;
+}
+
+void FGLTFLandscapeComponentDataInterface::VertexXYToTexelXY(int32 VertX, int32 VertY, int32& OutX, int32& OutY) const
+{
+	int32 SubNumX, SubNumY, SubX, SubY;
+	ComponentXYToSubsectionXY(VertX, VertY, SubNumX, SubNumY, SubX, SubY);
+
+	OutX = SubNumX * ComponentSubsectionSizeVerts + SubX;
+	OutY = SubNumY * ComponentSubsectionSizeVerts + SubY;
+}
+
+void FGLTFLandscapeComponentDataInterface::VertexIndexToXY(int32 VertexIndex, int32& OutX, int32& OutY) const
+{
+	OutX = VertexIndex % ComponentSizeVerts;
+	OutY = VertexIndex / ComponentSizeVerts;
+}
+
+int32 FGLTFLandscapeComponentDataInterface::VertexIndexToTexel(int32 VertexIndex) const
+{
+	int32 VertX, VertY;
+	VertexIndexToXY(VertexIndex, VertX, VertY);
+	int32 TexelX, TexelY;
+	VertexXYToTexelXY(VertX, VertY, TexelX, TexelY);
+	return TexelXYToIndex(TexelX, TexelY);
+}
+
+bool FGLTFLandscapeComponentDataInterface::GetWeightmapTextureData(ULandscapeLayerInfoObject* InLayerInfo,
+	TArray<uint8>& OutData,
+	bool bInUseEditingWeightmap, bool bInRemoveSubsectionDuplicates)
+{
+	int32 LayerIdx = INDEX_NONE;
+	const TArray<FWeightmapLayerAllocationInfo>& ComponentWeightmapLayerAllocations = Component.GetWeightmapLayerAllocations(bInUseEditingWeightmap);
+	const TArray<UTexture2D*>& ComponentWeightmapTextures = Component.GetWeightmapTextures(bInUseEditingWeightmap);
+
+	for (int32 Idx = 0; Idx < ComponentWeightmapLayerAllocations.Num(); Idx++)
+	{
+		if (ComponentWeightmapLayerAllocations[Idx].LayerInfo == InLayerInfo)
+		{
+			LayerIdx = Idx;
+			break;
+		}
+	}
+	if (LayerIdx < 0)
+	{
+		return false;
+	}
+	if (ComponentWeightmapLayerAllocations[LayerIdx].WeightmapTextureIndex >= ComponentWeightmapTextures.Num())
+	{
+		return false;
+	}
+	if (ComponentWeightmapLayerAllocations[LayerIdx].WeightmapTextureChannel >= 4)
+	{
+		return false;
+	}
+
+	// If requested to skip the duplicate row/col of texture data
+	int32 WeightmapSize = bInRemoveSubsectionDuplicates ?
+		((ComponentSubsectionSizeQuads * Component.NumSubsections) + 1) >> MipLevel :
+		((ComponentSubsectionSizeQuads + 1) * Component.NumSubsections) >> MipLevel;
+
+	OutData.Empty(FMath::Square(WeightmapSize));
+	OutData.AddUninitialized(FMath::Square(WeightmapSize));
+
+	// DataInterface Lock is a LockMipReadOnly on the texture
+	const TIndirectArray<FTexture2DMipMap>& WeightMips = ComponentWeightmapTextures[ComponentWeightmapLayerAllocations[LayerIdx].WeightmapTextureIndex]->GetPlatformMips();
+	if (MipLevel >= WeightMips.Num())
+	{
+		return false;
+	}
+	const FColor* WeightMipData = (FColor*)WeightMips[MipLevel].BulkData.LockReadOnly();
+
+	// Channel remapping
+	int32 ChannelOffsets[4] = { (int32)STRUCT_OFFSET(FColor, R), (int32)STRUCT_OFFSET(FColor, G), (int32)STRUCT_OFFSET(FColor, B), (int32)STRUCT_OFFSET(FColor, A) };
+
+	const uint8* SrcTextureData = (const uint8*)WeightMipData + ChannelOffsets[ComponentWeightmapLayerAllocations[LayerIdx].WeightmapTextureChannel];
+
+	for (int32 i = 0; i < FMath::Square(WeightmapSize); i++)
+	{
+		// If removing subsection duplicates, convert vertex to texel index
+		OutData[i] = bInRemoveSubsectionDuplicates ? SrcTextureData[VertexIndexToTexel(i) * sizeof(FColor)] : SrcTextureData[i * sizeof(FColor)];
+	}
+
+	WeightMips[MipLevel].BulkData.Unlock();
+
+	return true;
+}
+
+FColor FGLTFLandscapeComponentDataInterface::GetHeightData(int32 LocalX, int32 LocalY) const
+{
+	if (HeightMipData == nullptr)
+	{
+		return FColor(0, 0, 0, 0);
+	}
+
+	int32 TexelX, TexelY;
+	VertexXYToTexelXY(LocalX, LocalY, TexelX, TexelY);
+
+	return HeightMipData[TexelX + HeightmapComponentOffsetX + (TexelY + HeightmapComponentOffsetY) * HeightmapStride];
+}
+
+FColor FGLTFLandscapeComponentDataInterface::GetXYOffsetData(int32 LocalX, int32 LocalY) const
+{
+	if (XYOffsetMipData == nullptr)
+	{
+		return FColor(0,0,0,0);
+	}
+
+	const int32 WeightmapSize = ((ComponentSubsectionSizeQuads + 1) * ComponentNumSubsections) >> MipLevel;
+	int32 SubNumX;
+	int32 SubNumY;
+	int32 SubX;
+	int32 SubY;
+	ComponentXYToSubsectionXY(LocalX, LocalY, SubNumX, SubNumY, SubX, SubY);
+
+	return XYOffsetMipData[SubX + SubNumX * ComponentSubsectionSizeVerts + (SubY + SubNumY * ComponentSubsectionSizeVerts) * WeightmapSize];
+}
+
+void FGLTFLandscapeComponentDataInterface::GetXYOffset(int32 X, int32 Y, float& XOffset, float& YOffset) const
+{
+	if (XYOffsetMipData)
+	{
+		FColor Texel = GetXYOffsetData(X, Y);
+		XOffset = ((float)((Texel.R << 8) + Texel.G) - 32768.f) * LANDSCAPE_XYOFFSET_SCALE;
+		YOffset = ((float)((Texel.B << 8) + Texel.A) - 32768.f) * LANDSCAPE_XYOFFSET_SCALE;
+	}
+	else
+	{
+		XOffset = YOffset = 0.f;
+	}
+}
+
+void FGLTFLandscapeComponentDataInterface::GetPositionNormalUV(int32 LocalX, int32 LocalY,
+	FVector3f& Position,
+	FVector3f& Normal,
+	FVector2f& UV) const
+{
+	{
+		FColor Data = GetHeightData(LocalX, LocalY);
+		FVector TangentZ = LandscapeDataAccess::UnpackNormal(Data);
+		FVector TangentX = FVector(TangentZ.Z, 0.f, -TangentZ.X);
+		FVector TangentY = TangentZ ^ TangentX;
+
+		float Height = LandscapeDataAccess::UnpackHeight(Data);
+
+		const float ScaleFactor = (float)ComponentSizeQuads / (float)(ComponentSizeVerts - 1);
+		float XOffset, YOffset;
+		GetXYOffset(LocalX, LocalY, XOffset, YOffset);
+		Position = FVector3f(LocalX * ScaleFactor + XOffset, LocalY * ScaleFactor + YOffset, Height);
+
+		Normal = FVector3f(TangentZ.X, TangentZ.Y, TangentZ.Z);
+		UV = FVector2f(Component.GetSectionBase().X + LocalX * ScaleFactor, Component.GetSectionBase().Y + LocalY * ScaleFactor);
+	}
+}
\ No newline at end of file
diff --git a/Engine/Plugins/Enterprise/GLTFExporter/Source/GLTFExporter/Public/Builders/GLTFAnalyticsBuilder.h b/Engine/Plugins/Enterprise/GLTFExporter/Source/GLTFExporter/Public/Builders/GLTFAnalyticsBuilder.h
--- a/Engine/Plugins/Enterprise/GLTFExporter/Source/GLTFExporter/Public/Builders/GLTFAnalyticsBuilder.h
+++ b/Engine/Plugins/Enterprise/GLTFExporter/Source/GLTFExporter/Public/Builders/GLTFAnalyticsBuilder.h
@@ -16,6 +16,7 @@ class UMaterialInterface;
 class UTexture;
 class UCameraComponent;
 class ULightComponent;
+class ULandscapeComponent;
 
 class GLTFEXPORTER_API FGLTFAnalyticsBuilder : public FGLTFBufferBuilder
 {
@@ -30,6 +31,8 @@ protected:
 	void RecordActor(const AActor* Object);
 	void RecordStaticMesh(const UStaticMesh* Object);
 	void RecordSkeletalMesh(const USkeletalMesh* Object);
+	void RecordSplineStaticMesh(const UStaticMesh* Object);
+	void RecordLandscapeComponent(const ULandscapeComponent* Object);
 	void RecordLevelSequence(const ULevelSequence* Object);
 	void RecordAnimSequence(const UAnimSequence* Object);
 	void RecordMaterial(const UMaterialInterface* Object);
@@ -43,6 +46,8 @@ private:
 	TSet<const USceneComponent*>		ComponentsRecorded;
 	TSet<const UStaticMesh*>			StaticMeshesRecorded;
 	TSet<const USkeletalMesh*>			SkeletalMeshesRecorded;
+	TSet<const UStaticMesh*>			SplineStaticMeshesRecorded;
+	TSet<const ULandscapeComponent*>	LandscapeComponentsRecorded;
 	TSet<const ULevelSequence*>			LevelSequencesRecorded;
 	TSet<const UAnimSequence*>			AnimSequencesRecorded;
 	TSet<const UMaterialInterface*>		MaterialsRecorded;
diff --git a/Engine/Plugins/Enterprise/GLTFExporter/Source/GLTFExporter/Public/Builders/GLTFConvertBuilder.h b/Engine/Plugins/Enterprise/GLTFExporter/Source/GLTFExporter/Public/Builders/GLTFConvertBuilder.h
--- a/Engine/Plugins/Enterprise/GLTFExporter/Source/GLTFExporter/Public/Builders/GLTFConvertBuilder.h
+++ b/Engine/Plugins/Enterprise/GLTFExporter/Source/GLTFExporter/Public/Builders/GLTFConvertBuilder.h
@@ -20,6 +20,8 @@
 
 class UMeshComponent;
 class UPropertyValue;
+class FGLTFNormalArray;
+class FGLTFUVArray;
 
 class GLTFEXPORTER_API FGLTFConvertBuilder : public FGLTFAnalyticsBuilder
 {
@@ -32,20 +34,26 @@ public:
 	bool IsSelectedActor(const AActor* Object) const;
 	bool IsRootActor(const AActor* Actor) const;
 
+	FGLTFJsonAccessor* AddUniquePositionAccessor(const FPositionVertexBuffer* VertexBuffer);
 	FGLTFJsonAccessor* AddUniquePositionAccessor(const FGLTFMeshSection* MeshSection, const FPositionVertexBuffer* VertexBuffer);
 	FGLTFJsonAccessor* AddUniqueColorAccessor(const FGLTFMeshSection* MeshSection, const FColorVertexBuffer* VertexBuffer);
 	FGLTFJsonAccessor* AddUniqueNormalAccessor(const FGLTFMeshSection* MeshSection, const FStaticMeshVertexBuffer* VertexBuffer);
+	FGLTFJsonAccessor* AddUniqueNormalAccessor(const FGLTFNormalArray* Normals);
 	FGLTFJsonAccessor* AddUniqueTangentAccessor(const FGLTFMeshSection* MeshSection, const FStaticMeshVertexBuffer* VertexBuffer);
 	FGLTFJsonAccessor* AddUniqueUVAccessor(const FGLTFMeshSection* MeshSection, const FStaticMeshVertexBuffer* VertexBuffer, int32 UVIndex);
+	FGLTFJsonAccessor* AddUniqueUVAccessor(const FGLTFUVArray* UVs); /* Supports single UV channel */
 	FGLTFJsonAccessor* AddUniqueJointAccessor(const FGLTFMeshSection* MeshSection, const FSkinWeightVertexBuffer* VertexBuffer, int32 InfluenceOffset);
 	FGLTFJsonAccessor* AddUniqueWeightAccessor(const FGLTFMeshSection* MeshSection, const FSkinWeightVertexBuffer* VertexBuffer, int32 InfluenceOffset);
 	FGLTFJsonAccessor* AddUniqueIndexAccessor(const FGLTFMeshSection* MeshSection);
+	FGLTFJsonAccessor* AddUniqueIndexAccessor(const FGLTFIndexArray* IndexBuffer, const FString& MeshName);
 
 	FGLTFJsonMesh* AddUniqueMesh(const UStaticMesh* StaticMesh, const FGLTFMaterialArray& Materials = {}, int32 LODIndex = INDEX_NONE);
 	FGLTFJsonMesh* AddUniqueMesh(const USkeletalMesh* SkeletalMesh, const FGLTFMaterialArray& Materials = {}, int32 LODIndex = INDEX_NONE);
 	FGLTFJsonMesh* AddUniqueMesh(const UMeshComponent* MeshComponent, const FGLTFMaterialArray& Materials = {}, int32 LODIndex = INDEX_NONE);
 	FGLTFJsonMesh* AddUniqueMesh(const UStaticMeshComponent* StaticMeshComponent, const FGLTFMaterialArray& Materials = {}, int32 LODIndex = INDEX_NONE);
 	FGLTFJsonMesh* AddUniqueMesh(const USkeletalMeshComponent* SkeletalMeshComponent, const FGLTFMaterialArray& Materials = {}, int32 LODIndex = INDEX_NONE);
+	FGLTFJsonMesh* AddUniqueMesh(const USplineMeshComponent* SplineMeshComponent, const FGLTFMaterialArray& Materials = {}, int32 LODIndex = INDEX_NONE);
+	FGLTFJsonMesh* AddUniqueMesh(const ULandscapeComponent* LandscapeComponent, const UMaterialInterface* LandscapeMaterial = nullptr);
 
 	const FGLTFMeshData* AddUniqueMeshData(const UStaticMesh* StaticMesh, const UStaticMeshComponent* StaticMeshComponent = nullptr, int32 LODIndex = INDEX_NONE);
 	const FGLTFMeshData* AddUniqueMeshData(const USkeletalMesh* SkeletalMesh, const USkeletalMeshComponent* SkeletalMeshComponent = nullptr, int32 LODIndex = INDEX_NONE);
@@ -92,6 +100,13 @@ public:
 	void RegisterObjectVariant(const UObject* Object, const UPropertyValue* Property);
 	const TArray<const UPropertyValue*>* GetObjectVariants(const UObject* Object) const;
 
+	TUniquePtr<IGLTFPositionBufferConverterRaw> PositionBufferConverterRaw = MakeUnique<FGLTFPositionBufferConverterRaw>(*this);
+	TUniquePtr<IGLTFNormalBufferConverterRaw> NormalBufferConverterRaw = MakeUnique<FGLTFNormalBufferConverterRaw>(*this);
+	TUniquePtr<IGLTFUVBufferConverterRaw> UVBufferConverterRaw = MakeUnique<FGLTFUVBufferConverterRaw>(*this);
+	TUniquePtr<IGLTFIndexBufferConverterRaw> IndexBufferConverterRaw = MakeUnique<FGLTFIndexBufferConverterRaw>(*this);
+	TUniquePtr<IGLTFSplineMeshConverter> SplineMeshConverter = MakeUnique<FGLTFSplineMeshConverter>(*this);
+	TUniquePtr<IGLTFLandscapeMeshConverter> LandscapeConverter = MakeUnique<FGLTFLandscapeMeshConverter>(*this);
+
 	TUniquePtr<IGLTFPositionBufferConverter> PositionBufferConverter = MakeUnique<FGLTFPositionBufferConverter>(*this);
 	TUniquePtr<IGLTFColorBufferConverter> ColorBufferConverter = MakeUnique<FGLTFColorBufferConverter>(*this);
 	TUniquePtr<IGLTFNormalBufferConverter> NormalBufferConverter = MakeUnique<FGLTFNormalBufferConverter>(*this);
diff --git a/Engine/Plugins/Enterprise/GLTFExporter/Source/GLTFExporter/Public/Converters/GLTFAccessorConverters.h b/Engine/Plugins/Enterprise/GLTFExporter/Source/GLTFExporter/Public/Converters/GLTFAccessorConverters.h
--- a/Engine/Plugins/Enterprise/GLTFExporter/Source/GLTFExporter/Public/Converters/GLTFAccessorConverters.h
+++ b/Engine/Plugins/Enterprise/GLTFExporter/Source/GLTFExporter/Public/Converters/GLTFAccessorConverters.h
@@ -11,6 +11,8 @@ class FPositionVertexBuffer;
 class FColorVertexBuffer;
 class FStaticMeshVertexBuffer;
 class FSkinWeightVertexBuffer;
+class FGLTFNormalArray;
+class FGLTFUVArray;
 
 typedef TGLTFConverter<FGLTFJsonAccessor*, const FGLTFMeshSection*, const FPositionVertexBuffer*> IGLTFPositionBufferConverter;
 typedef TGLTFConverter<FGLTFJsonAccessor*, const FGLTFMeshSection*, const FColorVertexBuffer*> IGLTFColorBufferConverter;
@@ -20,6 +22,11 @@ typedef TGLTFConverter<FGLTFJsonAccessor*, const FGLTFMeshSection*, const FStati
 typedef TGLTFConverter<FGLTFJsonAccessor*, const FGLTFMeshSection*, const FSkinWeightVertexBuffer*, uint32> IGLTFBoneIndexBufferConverter;
 typedef TGLTFConverter<FGLTFJsonAccessor*, const FGLTFMeshSection*, const FSkinWeightVertexBuffer*, uint32> IGLTFBoneWeightBufferConverter;
 typedef TGLTFConverter<FGLTFJsonAccessor*, const FGLTFMeshSection*> IGLTFIndexBufferConverter;
+typedef TGLTFConverter<FGLTFJsonAccessor*, const FPositionVertexBuffer*> IGLTFPositionBufferConverterRaw;
+typedef TGLTFConverter<FGLTFJsonAccessor*, const FGLTFIndexArray*, FString> IGLTFIndexBufferConverterRaw;
+typedef TGLTFConverter<FGLTFJsonAccessor*, const FGLTFNormalArray*> IGLTFNormalBufferConverterRaw;
+typedef TGLTFConverter<FGLTFJsonAccessor*, const FGLTFUVArray*> IGLTFUVBufferConverterRaw;
+
 
 class GLTFEXPORTER_API FGLTFPositionBufferConverter : public FGLTFBuilderContext, public IGLTFPositionBufferConverter
 {
@@ -141,9 +148,56 @@ public:
 protected:
 
 	virtual FGLTFJsonAccessor* Convert(const FGLTFMeshSection* MeshSection) override;
+};
+
+
+class GLTFEXPORTER_API FGLTFPositionBufferConverterRaw : public FGLTFBuilderContext, public IGLTFPositionBufferConverterRaw
+{
+public:
+
+	using FGLTFBuilderContext::FGLTFBuilderContext;
+
+protected:
+
+	virtual FGLTFJsonAccessor* Convert(const FPositionVertexBuffer* VertexBuffer) override;
+};
+
+class GLTFEXPORTER_API FGLTFIndexBufferConverterRaw : public FGLTFBuilderContext, public IGLTFIndexBufferConverterRaw
+{
+public:
+
+	using FGLTFBuilderContext::FGLTFBuilderContext;
+
+protected:
+	virtual FGLTFJsonAccessor* Convert(const FGLTFIndexArray* IndexBuffer, FString MeshName) override;
+};
+
+class GLTFEXPORTER_API FGLTFNormalBufferConverterRaw : public FGLTFBuilderContext, public IGLTFNormalBufferConverterRaw
+{
+public:
+
+	using FGLTFBuilderContext::FGLTFBuilderContext;
+
+protected:
+
+	virtual FGLTFJsonAccessor* Convert(const FGLTFNormalArray* NormalSource) override;
 
 private:
 
-	template <typename IndexType>
-	FGLTFJsonAccessor* Convert(const FGLTFMeshSection* MeshSection) const;
+	FGLTFJsonBufferView* ConvertBufferView(const FGLTFNormalArray* NormalSource);
 };
+
+class GLTFEXPORTER_API FGLTFUVBufferConverterRaw : public FGLTFBuilderContext, public IGLTFUVBufferConverterRaw
+{
+public:
+
+	using FGLTFBuilderContext::FGLTFBuilderContext;
+
+protected:
+
+	virtual FGLTFJsonAccessor* Convert(const FGLTFUVArray* UVSource) override;
+
+private:
+
+	FGLTFJsonBufferView* ConvertBufferView(const FGLTFUVArray* UVSource);
+};
\ No newline at end of file
diff --git a/Engine/Plugins/Enterprise/GLTFExporter/Source/GLTFExporter/Public/Converters/GLTFMeshConverters.h b/Engine/Plugins/Enterprise/GLTFExporter/Source/GLTFExporter/Public/Converters/GLTFMeshConverters.h
--- a/Engine/Plugins/Enterprise/GLTFExporter/Source/GLTFExporter/Public/Converters/GLTFMeshConverters.h
+++ b/Engine/Plugins/Enterprise/GLTFExporter/Source/GLTFExporter/Public/Converters/GLTFMeshConverters.h
@@ -12,9 +12,13 @@ class UStaticMesh;
 class UStaticMeshComponent;
 class USkeletalMesh;
 class USkeletalMeshComponent;
+class USplineMeshComponent;
+class ULandscapeComponent;
 
 typedef TGLTFConverter<FGLTFJsonMesh*, const UStaticMesh*, const UStaticMeshComponent*, FGLTFMaterialArray, int32> IGLTFStaticMeshConverter;
 typedef TGLTFConverter<FGLTFJsonMesh*, const USkeletalMesh*, const USkeletalMeshComponent*, FGLTFMaterialArray, int32> IGLTFSkeletalMeshConverter;
+typedef TGLTFConverter<FGLTFJsonMesh*, const UStaticMesh*, const USplineMeshComponent*, FGLTFMaterialArray, int32> IGLTFSplineMeshConverter;
+typedef TGLTFConverter<FGLTFJsonMesh*, const ULandscapeComponent*, const UMaterialInterface*> IGLTFLandscapeMeshConverter;
 
 class GLTFEXPORTER_API FGLTFStaticMeshConverter : public FGLTFBuilderContext, public IGLTFStaticMeshConverter
 {
@@ -49,3 +53,35 @@ private:
 
 	FGLTFSkeletalMeshSectionConverter MeshSectionConverter;
 };
+
+class GLTFEXPORTER_API FGLTFSplineMeshConverter : public FGLTFBuilderContext, public IGLTFSplineMeshConverter
+{
+public:
+
+	using FGLTFBuilderContext::FGLTFBuilderContext;
+
+protected:
+
+	virtual void Sanitize(const UStaticMesh*& StaticMesh, const USplineMeshComponent*& SplineMeshComponent, FGLTFMaterialArray& Materials, int32& LODIndex) override;
+
+	virtual FGLTFJsonMesh* Convert(const UStaticMesh* StaticMesh, const USplineMeshComponent* SplineMeshComponent, FGLTFMaterialArray Materials, int32 LODIndex) override;
+
+private:
+
+	FGLTFStaticMeshSectionConverter MeshSectionConverter;
+};
+
+class GLTFEXPORTER_API FGLTFLandscapeMeshConverter : public FGLTFBuilderContext, public IGLTFLandscapeMeshConverter
+{
+public:
+
+	using FGLTFBuilderContext::FGLTFBuilderContext;
+
+protected:
+
+	virtual void Sanitize(const ULandscapeComponent*& LandscapeComponent, const UMaterialInterface*& LandscapeMaterial) override;
+
+	virtual FGLTFJsonMesh* Convert(const ULandscapeComponent* LandscapeComponent, const UMaterialInterface* LandscapeMaterial) override;
+
+private:
+};
\ No newline at end of file
diff --git a/Engine/Plugins/Enterprise/GLTFExporter/Source/GLTFExporter/Public/Converters/GLTFNormalArray.h b/Engine/Plugins/Enterprise/GLTFExporter/Source/GLTFExporter/Public/Converters/GLTFNormalArray.h
new file mode 100644
--- /dev/null
+++ b/Engine/Plugins/Enterprise/GLTFExporter/Source/GLTFExporter/Public/Converters/GLTFNormalArray.h
@@ -0,0 +1,24 @@
+// Copyright Epic Games, Inc. All Rights Reserved.
+
+#pragma once
+
+#include "CoreMinimal.h"
+
+// TODO: should we add GLTFEXPORTER_API ?
+class FGLTFNormalArray : public TArray<FVector3f>
+{
+public:
+
+	using TArray::TArray;
+	using TArray::operator=;
+
+	friend uint32 GetTypeHash(const FGLTFNormalArray& NormalArray)
+	{
+		uint32 Hash = GetTypeHash(NormalArray.Num());
+		for (const FVector3f& Normal : NormalArray)
+		{
+			Hash = HashCombine(Hash, GetTypeHash(Normal));
+		}
+		return Hash;
+	}
+};
\ No newline at end of file
diff --git a/Engine/Plugins/Enterprise/GLTFExporter/Source/GLTFExporter/Public/Converters/GLTFUVArray.h b/Engine/Plugins/Enterprise/GLTFExporter/Source/GLTFExporter/Public/Converters/GLTFUVArray.h
new file mode 100644
--- /dev/null
+++ b/Engine/Plugins/Enterprise/GLTFExporter/Source/GLTFExporter/Public/Converters/GLTFUVArray.h
@@ -0,0 +1,24 @@
+// Copyright Epic Games, Inc. All Rights Reserved.
+
+#pragma once
+
+#include "CoreMinimal.h"
+
+// TODO: should we add GLTFEXPORTER_API ?
+class FGLTFUVArray : public TArray<FVector2f>
+{
+public:
+
+	using TArray::TArray;
+	using TArray::operator=;
+
+	friend uint32 GetTypeHash(const FGLTFUVArray& UVArray)
+	{
+		uint32 Hash = GetTypeHash(UVArray.Num());
+		for (const FVector2f& UV : UVArray)
+		{
+			Hash = HashCombine(Hash, GetTypeHash(UV));
+		}
+		return Hash;
+	}
+};
\ No newline at end of file
diff --git a/Engine/Plugins/Enterprise/GLTFExporter/Source/GLTFExporter/Public/Utilities/GLTFLandscapeComponentDataInterface.h b/Engine/Plugins/Enterprise/GLTFExporter/Source/GLTFExporter/Public/Utilities/GLTFLandscapeComponentDataInterface.h
new file mode 100644
--- /dev/null
+++ b/Engine/Plugins/Enterprise/GLTFExporter/Source/GLTFExporter/Public/Utilities/GLTFLandscapeComponentDataInterface.h
@@ -0,0 +1,59 @@
+// Copyright Epic Games, Inc. All Rights Reserved.
+
+#pragma once
+
+#include "CoreMinimal.h"
+
+class ULandscapeComponent;
+class ULandscapeLayerInfoObject;
+
+struct FGLTFLandscapeComponentDataInterface
+{
+protected:
+	int32 ComponentSubsectionSizeQuads = 0;
+	int32 ComponentSizeQuads = 0;
+	int32 ComponentSizeVerts = 0;
+	int32 ComponentSubsectionSizeVerts = 0;
+	int32 ComponentNumSubsections = 0;
+
+	TArray64<uint8> HeightMipDataSource;
+	FColor* HeightMipData = nullptr;
+	int32 HeightmapStride = 0;
+	int32 HeightmapComponentOffsetX = 0;
+	int32 HeightmapComponentOffsetY = 0;
+	int32 HeightmapSubsectionOffset = 0;
+
+	TArray64<uint8> XYOffsetMipDataSource;
+	FColor* XYOffsetMipData = nullptr;
+
+	const ULandscapeComponent& Component;
+	int32 MipLevel;
+
+public:
+	FGLTFLandscapeComponentDataInterface(const ULandscapeComponent& Component, int32 ExportLOD);
+
+	void ComponentXYToSubsectionXY(int32 CompX, int32 CompY, int32& SubNumX, int32& SubNumY, int32& SubX, int32& SubY) const;
+
+	int32 TexelXYToIndex(int32 TexelX, int32 TexelY) const;
+
+	void VertexXYToTexelXY(int32 VertX, int32 VertY, int32& OutX, int32& OutY) const;
+
+	void VertexIndexToXY(int32 VertexIndex, int32& OutX, int32& OutY) const;
+
+	int32 VertexIndexToTexel(int32 VertexIndex) const;
+
+	bool GetWeightmapTextureData(ULandscapeLayerInfoObject* InLayerInfo,
+		TArray<uint8>& OutData,
+		bool bInUseEditingWeightmap = false, bool bInRemoveSubsectionDuplicates = false);
+
+	FColor GetHeightData(int32 LocalX, int32 LocalY) const;
+
+	FColor GetXYOffsetData(int32 LocalX, int32 LocalY) const;
+
+	void GetXYOffset(int32 X, int32 Y, float& XOffset, float& YOffset) const;
+
+	void GetPositionNormalUV(int32 LocalX, int32 LocalY,
+		FVector3f& Position,
+		FVector3f& Normal,
+		FVector2f& UV) const;
+};
\ No newline at end of file
