Issue 420216: PVS Studio warnings in GLTFExporter

#rnx

diff --git a/Engine/Plugins/Enterprise/GLTFExporter/Source/GLTFExporter/Private/Converters/GLTFMeshConverters.cpp b/Engine/Plugins/Enterprise/GLTFExporter/Source/GLTFExporter/Private/Converters/GLTFMeshConverters.cpp
--- a/Engine/Plugins/Enterprise/GLTFExporter/Source/GLTFExporter/Private/Converters/GLTFMeshConverters.cpp
+++ b/Engine/Plugins/Enterprise/GLTFExporter/Source/GLTFExporter/Private/Converters/GLTFMeshConverters.cpp
@@ -110,7 +110,7 @@ void FGLTFSplineMeshConverter::Sanitize(const UStaticMesh*& StaticMesh, const US
 FGLTFJsonMesh* FGLTFSplineMeshConverter::Convert(const UStaticMesh* StaticMesh, const USplineMeshComponent* SplineMeshComponent, FGLTFMaterialArray Materials, int32 LODIndex)
 {
 #if !WITH_EDITOR
-	if (!StaticMesh->bAllowCPUAccess)
+	if (StaticMesh && !StaticMesh->bAllowCPUAccess)
 	{
 		Builder.LogSuggestion(FString::Printf(
 			TEXT("Export of mesh %s can in runtime be speed-up by checking 'Allow CPU Access' in asset settings"),
@@ -118,7 +118,7 @@ FGLTFJsonMesh* FGLTFSplineMeshConverter::Convert(const UStaticMesh* StaticMesh,
 	}
 #endif
 
-	const int32 MaterialCount = FGLTFMeshUtilities::GetMaterials(StaticMesh).Num();
+	const int32 MaterialCount = StaticMesh ? FGLTFMeshUtilities::GetMaterials(StaticMesh).Num() : 0;
 	FGLTFJsonMesh* JsonMesh = Builder.AddMesh();
 	JsonMesh->Primitives.AddDefaulted(MaterialCount);
 
diff --git a/Engine/Plugins/Enterprise/GLTFExporter/Source/GLTFExporter/Private/Converters/GLTFSceneConverters.cpp b/Engine/Plugins/Enterprise/GLTFExporter/Source/GLTFExporter/Private/Converters/GLTFSceneConverters.cpp
--- a/Engine/Plugins/Enterprise/GLTFExporter/Source/GLTFExporter/Private/Converters/GLTFSceneConverters.cpp
+++ b/Engine/Plugins/Enterprise/GLTFExporter/Source/GLTFExporter/Private/Converters/GLTFSceneConverters.cpp
@@ -18,17 +18,12 @@ FGLTFJsonScene* FGLTFSceneConverter::Convert(const UWorld* World)
 	{
 		for (const ULevel* Level : Levels)
 		{
-			if (Level->IsInstancedLevel() && !Level->IsCurrentLevel())
+			if (Level == nullptr || (Level->IsInstancedLevel() && !Level->IsCurrentLevel()))
 			{
 				//LevelInstances are handled in FGLTFComponentConverter::ConvertComponentSpecialization in the ULevelInstanceComponent handling branch.
 				continue;
 			}
 
-			if (Level == nullptr)
-			{
-				continue;
-			}
-
 			// TODO: add support for exporting Level->Model?
 
 			TArray<AActor*> LevelActors(Level->Actors); // iterate using copy, in case new temporary actors are added during conversion
diff --git a/Engine/Plugins/Enterprise/GLTFExporter/Source/GLTFExporter/Private/Tasks/GLTFDelayedMeshTasks.cpp b/Engine/Plugins/Enterprise/GLTFExporter/Source/GLTFExporter/Private/Tasks/GLTFDelayedMeshTasks.cpp
--- a/Engine/Plugins/Enterprise/GLTFExporter/Source/GLTFExporter/Private/Tasks/GLTFDelayedMeshTasks.cpp
+++ b/Engine/Plugins/Enterprise/GLTFExporter/Source/GLTFExporter/Private/Tasks/GLTFDelayedMeshTasks.cpp
@@ -398,10 +398,9 @@ void FGLTFDelayedSplineMeshTask::Process()
 		FGLTFJsonPrimitive& JsonPrimitive = JsonMesh->Primitives[MaterialIndex];
 		JsonPrimitive.Indices = Builder.AddUniqueIndexAccessor(ConvertedSection);
 
-		FPositionVertexBuffer* TransformedPositionBuffer = nullptr;
+		FPositionVertexBuffer* TransformedPositionBuffer = new FPositionVertexBuffer();
 		
 		{//fix for Splines:
-			TransformedPositionBuffer = new FPositionVertexBuffer();
 			TransformedPositionBuffer->Init(PositionBuffer.GetNumVertices(), true);
 
 			const uint32 VertexCount = PositionBuffer.GetNumVertices();
@@ -422,7 +421,7 @@ void FGLTFDelayedSplineMeshTask::Process()
 			}
 		}
 
-		JsonPrimitive.Attributes.Position = Builder.AddUniquePositionAccessor(ConvertedSection, TransformedPositionBuffer != nullptr ? TransformedPositionBuffer : &PositionBuffer);
+		JsonPrimitive.Attributes.Position = Builder.AddUniquePositionAccessor(ConvertedSection, TransformedPositionBuffer);
 		if (JsonPrimitive.Attributes.Position == nullptr)
 		{
 			Builder.LogError(
