[glTF Exporter] Add support for export of KHR_materials_anisotropy.
[glTF Exporter] Anisotropy support added for glTF Proxy materials.
[glTF Exporter] Iridescence support added for glTF Proxy materials.

#rnx

#virtualized

diff --git a/Engine/Plugins/Enterprise/GLTFExporter/Source/GLTFExporter/Private/Converters/GLTFMaterialProxyFactory.cpp b/Engine/Plugins/Enterprise/GLTFExporter/Source/GLTFExporter/Private/Converters/GLTFMaterialProxyFactory.cpp
--- a/Engine/Plugins/Enterprise/GLTFExporter/Source/GLTFExporter/Private/Converters/GLTFMaterialProxyFactory.cpp
+++ b/Engine/Plugins/Enterprise/GLTFExporter/Source/GLTFExporter/Private/Converters/GLTFMaterialProxyFactory.cpp
@@ -148,6 +148,20 @@ void FGLTFMaterialProxyFactory::SetProxyParameters(UMaterialInstanceConstant* Pr
 			SetProxyParameter(ProxyMaterial, FGLTFProxyMaterialInfo::TransmissionFactor, JsonMaterial.Transmission.Factor);
 			SetProxyParameter(ProxyMaterial, FGLTFProxyMaterialInfo::TransmissionTexture, JsonMaterial.Transmission.Texture);
 		}
+
+		if (JsonMaterial.ShadingModel != EGLTFJsonShadingModel::SpecularGlossiness)
+		{
+			SetProxyParameter(ProxyMaterial, FGLTFProxyMaterialInfo::IridescenceFactor, JsonMaterial.Iridescence.IridescenceFactor);
+			SetProxyParameter(ProxyMaterial, FGLTFProxyMaterialInfo::IridescenceTexture, JsonMaterial.Iridescence.IridescenceTexture);
+			SetProxyParameter(ProxyMaterial, FGLTFProxyMaterialInfo::IridescenceIOR, JsonMaterial.Iridescence.IridescenceIOR);
+			SetProxyParameter(ProxyMaterial, FGLTFProxyMaterialInfo::IridescenceThicknessMinimum, JsonMaterial.Iridescence.IridescenceThicknessMinimum);
+			SetProxyParameter(ProxyMaterial, FGLTFProxyMaterialInfo::IridescenceThicknessMaximum, JsonMaterial.Iridescence.IridescenceThicknessMaximum);
+			SetProxyParameter(ProxyMaterial, FGLTFProxyMaterialInfo::IridescenceThicknessTexture, JsonMaterial.Iridescence.IridescenceThicknessTexture);
+
+			SetProxyParameter(ProxyMaterial, FGLTFProxyMaterialInfo::AnisotropyStrength, JsonMaterial.Anisotropy.AnisotropyStrength);
+			SetProxyParameter(ProxyMaterial, FGLTFProxyMaterialInfo::AnisotropyRotation, JsonMaterial.Anisotropy.AnisotropyRotation);
+			SetProxyParameter(ProxyMaterial, FGLTFProxyMaterialInfo::AnisotropyTexture, JsonMaterial.Anisotropy.AnisotropyTexture);
+		}
 	}
 }
 
diff --git a/Engine/Plugins/Enterprise/GLTFExporter/Source/GLTFExporter/Private/Converters/GLTFMaterialUtilities.cpp b/Engine/Plugins/Enterprise/GLTFExporter/Source/GLTFExporter/Private/Converters/GLTFMaterialUtilities.cpp
--- a/Engine/Plugins/Enterprise/GLTFExporter/Source/GLTFExporter/Private/Converters/GLTFMaterialUtilities.cpp
+++ b/Engine/Plugins/Enterprise/GLTFExporter/Source/GLTFExporter/Private/Converters/GLTFMaterialUtilities.cpp
@@ -1001,6 +1001,11 @@ void FGLTFImportMaterialMatchMakingHelper::Process()
 			GetValue(Inputs::IridescenceThicknessMinimum, JsonMaterial.Iridescence.IridescenceThicknessMinimum);
 			GetValue(Inputs::IridescenceThicknessMaximum, JsonMaterial.Iridescence.IridescenceThicknessMaximum);
 			GetValue(Inputs::IridescenceThicknessTexture, JsonMaterial.Iridescence.IridescenceThicknessTexture);
+
+			//Anisotropy
+			GetValue(Inputs::AnisotropyStrength, JsonMaterial.Anisotropy.AnisotropyStrength);
+			GetValue(Inputs::AnisotropyRotation, JsonMaterial.Anisotropy.AnisotropyRotation);
+			GetValue(Inputs::AnisotropyTexture, JsonMaterial.Anisotropy.AnisotropyTexture);
 		}
 
 		//Importer does not support Emissive for Transmission at the moment
diff --git a/Engine/Plugins/Enterprise/GLTFExporter/Source/GLTFExporter/Private/Json/GLTFJsonMaterial.cpp b/Engine/Plugins/Enterprise/GLTFExporter/Source/GLTFExporter/Private/Json/GLTFJsonMaterial.cpp
--- a/Engine/Plugins/Enterprise/GLTFExporter/Source/GLTFExporter/Private/Json/GLTFJsonMaterial.cpp
+++ b/Engine/Plugins/Enterprise/GLTFExporter/Source/GLTFExporter/Private/Json/GLTFJsonMaterial.cpp
@@ -239,6 +239,24 @@ void FGLTFJsonIridescenceExtension::WriteObject(IGLTFJsonWriter& Writer) const
 	}
 }
 
+void FGLTFJsonAnisotropyExtension::WriteObject(IGLTFJsonWriter& Writer) const
+{
+	if (!FMath::IsNearlyEqual(AnisotropyStrength, 0.f))
+	{
+		Writer.Write(TEXT("anisotropyStrength"), AnisotropyStrength);
+	}
+
+	if (!FMath::IsNearlyEqual(AnisotropyRotation, 0.f))
+	{
+		Writer.Write(TEXT("anisotropyRotation"), AnisotropyRotation);
+	}
+
+	if (AnisotropyTexture.Index != nullptr)
+	{
+		Writer.Write(TEXT("anisotropyTexture"), AnisotropyTexture);
+	}
+}
+
 void FGLTFJsonMaterial::WriteObject(IGLTFJsonWriter& Writer) const
 {
 	if (!Name.IsEmpty())
@@ -296,7 +314,9 @@ void FGLTFJsonMaterial::WriteObject(IGLTFJsonWriter& Writer) const
 		Specular.HasValue() || 
 		IOR.HasValue() ||
 		Sheen.HasValue() ||
-		Iridescence.HasValue())
+		Iridescence.HasValue() ||
+		Anisotropy.HasValue()
+		)
 	{
 		Writer.StartExtensions();
 
@@ -320,6 +340,11 @@ void FGLTFJsonMaterial::WriteObject(IGLTFJsonWriter& Writer) const
 			Writer.Write(EGLTFJsonExtension::KHR_MaterialsIridescence, Iridescence);
 		}
 
+		if (ShadingModel != EGLTFJsonShadingModel::Unlit && ShadingModel != EGLTFJsonShadingModel::SpecularGlossiness && Anisotropy.HasValue())
+		{
+			Writer.Write(EGLTFJsonExtension::KHR_MaterialsAnisotropy, Anisotropy);
+		}
+
 		if (ShadingModel == EGLTFJsonShadingModel::Unlit)
 		{
 			// Write empty object
@@ -339,6 +364,13 @@ void FGLTFJsonMaterial::WriteObject(IGLTFJsonWriter& Writer) const
 			Writer.Write(EGLTFJsonExtension::KHR_MaterialsSpecularGlossiness, PBRSpecularGlossiness);
 		}
 
+		if (ShadingModel != EGLTFJsonShadingModel::SpecularGlossiness &&
+			ShadingModel != EGLTFJsonShadingModel::Unlit && 
+			Anisotropy.HasValue())
+		{
+
+		}
+
 		if (HasEmissiveStrength)
 		{
 			Writer.StartExtension(EGLTFJsonExtension::KHR_MaterialsEmissiveStrength);
diff --git a/Engine/Plugins/Enterprise/GLTFExporter/Source/GLTFExporter/Private/Materials/GLTFProxyMaterialInfo.cpp b/Engine/Plugins/Enterprise/GLTFExporter/Source/GLTFExporter/Private/Materials/GLTFProxyMaterialInfo.cpp
--- a/Engine/Plugins/Enterprise/GLTFExporter/Source/GLTFExporter/Private/Materials/GLTFProxyMaterialInfo.cpp
+++ b/Engine/Plugins/Enterprise/GLTFExporter/Source/GLTFExporter/Private/Materials/GLTFProxyMaterialInfo.cpp
@@ -40,3 +40,14 @@ const FGLTFProxyMaterialTextureParameterInfo FGLTFProxyMaterialInfo::SheenRoughn
 
 const TGLTFProxyMaterialParameterInfo<float> FGLTFProxyMaterialInfo::TransmissionFactor = { TEXT("Transmission Factor") };
 const FGLTFProxyMaterialTextureParameterInfo FGLTFProxyMaterialInfo::TransmissionTexture = { TEXT("Transmission") };
+
+const TGLTFProxyMaterialParameterInfo<float> FGLTFProxyMaterialInfo::IridescenceFactor = { TEXT("Iridescence Factor") };
+const TGLTFProxyMaterialParameterInfo<float> FGLTFProxyMaterialInfo::IridescenceIOR = { TEXT("Iridescence IOR") };
+const FGLTFProxyMaterialTextureParameterInfo FGLTFProxyMaterialInfo::IridescenceTexture = { TEXT("Iridescence") };
+const TGLTFProxyMaterialParameterInfo<float> FGLTFProxyMaterialInfo::IridescenceThicknessMinimum = { TEXT("IridescenceThickness Minimum") };
+const TGLTFProxyMaterialParameterInfo<float> FGLTFProxyMaterialInfo::IridescenceThicknessMaximum = { TEXT("IridescenceThickness Maximum") };
+const FGLTFProxyMaterialTextureParameterInfo FGLTFProxyMaterialInfo::IridescenceThicknessTexture = { TEXT("IridescenceThickness") };
+
+const TGLTFProxyMaterialParameterInfo<float> FGLTFProxyMaterialInfo::AnisotropyStrength = { TEXT("Anisotropy Strength") };
+const TGLTFProxyMaterialParameterInfo<float> FGLTFProxyMaterialInfo::AnisotropyRotation = { TEXT("Anisotropy Rotation") };
+const FGLTFProxyMaterialTextureParameterInfo FGLTFProxyMaterialInfo::AnisotropyTexture = { TEXT("Anisotropy") };
diff --git a/Engine/Plugins/Enterprise/GLTFExporter/Source/GLTFExporter/Private/Materials/GLTFProxyMaterialInfo.h b/Engine/Plugins/Enterprise/GLTFExporter/Source/GLTFExporter/Private/Materials/GLTFProxyMaterialInfo.h
--- a/Engine/Plugins/Enterprise/GLTFExporter/Source/GLTFExporter/Private/Materials/GLTFProxyMaterialInfo.h
+++ b/Engine/Plugins/Enterprise/GLTFExporter/Source/GLTFExporter/Private/Materials/GLTFProxyMaterialInfo.h
@@ -46,4 +46,15 @@ public:
 
 	static const TGLTFProxyMaterialParameterInfo<float> TransmissionFactor;
 	static const FGLTFProxyMaterialTextureParameterInfo TransmissionTexture; //Only using Red Channel
+
+	static const TGLTFProxyMaterialParameterInfo<float> IridescenceFactor;
+	static const TGLTFProxyMaterialParameterInfo<float> IridescenceIOR;
+	static const FGLTFProxyMaterialTextureParameterInfo IridescenceTexture;
+	static const TGLTFProxyMaterialParameterInfo<float> IridescenceThicknessMinimum;
+	static const TGLTFProxyMaterialParameterInfo<float> IridescenceThicknessMaximum;
+	static const FGLTFProxyMaterialTextureParameterInfo IridescenceThicknessTexture;
+
+	static const TGLTFProxyMaterialParameterInfo<float> AnisotropyStrength;
+	static const TGLTFProxyMaterialParameterInfo<float> AnisotropyRotation;
+	static const FGLTFProxyMaterialTextureParameterInfo AnisotropyTexture;
 };
diff --git a/Engine/Plugins/Enterprise/GLTFExporter/Source/GLTFExporter/Private/Tasks/GLTFDelayedMaterialTasks.cpp b/Engine/Plugins/Enterprise/GLTFExporter/Source/GLTFExporter/Private/Tasks/GLTFDelayedMaterialTasks.cpp
--- a/Engine/Plugins/Enterprise/GLTFExporter/Source/GLTFExporter/Private/Tasks/GLTFDelayedMaterialTasks.cpp
+++ b/Engine/Plugins/Enterprise/GLTFExporter/Source/GLTFExporter/Private/Tasks/GLTFDelayedMaterialTasks.cpp
@@ -235,6 +235,30 @@ void FGLTFDelayedMaterialTask::Process()
 					Builder.LogWarning(FString::Printf(TEXT("Failed to export %s and %s for material %s"), *FuzzColorProperty.ToString(), *ClothProperty.ToString(), *Material->GetName()));
 				}
 			}
+
+			if (JsonMaterial->ShadingModel != EGLTFJsonShadingModel::SpecularGlossiness)
+			{
+				const FMaterialPropertyEx AnisotropyProperty = MP_Anisotropy;
+				if (IsPropertyNonDefault(AnisotropyProperty))
+				{
+					if (!TryGetSourceTexture(JsonMaterial->Anisotropy.AnisotropyTexture, AnisotropyProperty, DefaultColorInputMasks))
+					{
+						if (!TryGetBakedMaterialProperty(JsonMaterial->Anisotropy.AnisotropyTexture, JsonMaterial->Anisotropy.AnisotropyStrength, 1.0f, AnisotropyProperty,
+							[](FColor& Color, const uint8& ChannelValueToMove)
+							{
+								Color.B = ChannelValueToMove;
+								
+								Color.R = 255;
+								Color.G = 255;
+								Color.A = 255;
+							}
+						))
+						{
+							Builder.LogWarning(FString::Printf(TEXT("Failed to export %s for material %s"), *AnisotropyProperty.ToString(), *Material->GetName()));
+						}
+					}
+				}
+			}
 		}
 	}
 
@@ -346,6 +370,20 @@ void FGLTFDelayedMaterialTask::GetProxyParameters(FGLTFJsonMaterial& OutMaterial
 			GetProxyParameter(FGLTFProxyMaterialInfo::TransmissionFactor, OutMaterial.Transmission.Factor);
 			GetProxyParameter(FGLTFProxyMaterialInfo::TransmissionTexture, OutMaterial.Transmission.Texture);
 		}
+
+		if (OutMaterial.ShadingModel != EGLTFJsonShadingModel::SpecularGlossiness)
+		{
+			GetProxyParameter(FGLTFProxyMaterialInfo::IridescenceFactor, OutMaterial.Iridescence.IridescenceFactor);
+			GetProxyParameter(FGLTFProxyMaterialInfo::IridescenceTexture, OutMaterial.Iridescence.IridescenceTexture);
+			GetProxyParameter(FGLTFProxyMaterialInfo::IridescenceIOR, OutMaterial.Iridescence.IridescenceIOR);
+			GetProxyParameter(FGLTFProxyMaterialInfo::IridescenceThicknessMinimum, OutMaterial.Iridescence.IridescenceThicknessMinimum);
+			GetProxyParameter(FGLTFProxyMaterialInfo::IridescenceThicknessMaximum, OutMaterial.Iridescence.IridescenceThicknessMaximum);
+			GetProxyParameter(FGLTFProxyMaterialInfo::IridescenceThicknessTexture, OutMaterial.Iridescence.IridescenceThicknessTexture);
+
+			GetProxyParameter(FGLTFProxyMaterialInfo::AnisotropyStrength, OutMaterial.Anisotropy.AnisotropyStrength);
+			GetProxyParameter(FGLTFProxyMaterialInfo::AnisotropyRotation, OutMaterial.Anisotropy.AnisotropyRotation);
+			GetProxyParameter(FGLTFProxyMaterialInfo::AnisotropyTexture, OutMaterial.Anisotropy.AnisotropyTexture);
+		}
 	}
 }
 
diff --git a/Engine/Plugins/Enterprise/GLTFExporter/Source/GLTFExporter/Private/Utilities/GLTFJsonUtilities.cpp b/Engine/Plugins/Enterprise/GLTFExporter/Source/GLTFExporter/Private/Utilities/GLTFJsonUtilities.cpp
--- a/Engine/Plugins/Enterprise/GLTFExporter/Source/GLTFExporter/Private/Utilities/GLTFJsonUtilities.cpp
+++ b/Engine/Plugins/Enterprise/GLTFExporter/Source/GLTFExporter/Private/Utilities/GLTFJsonUtilities.cpp
@@ -16,6 +16,7 @@ const TCHAR* FGLTFJsonUtilities::GetValue(EGLTFJsonExtension Enum)
 		case EGLTFJsonExtension::KHR_MaterialsTransmission:           return TEXT("KHR_materials_transmission");
 		case EGLTFJsonExtension::KHR_MaterialsSpecularGlossiness:     return TEXT("KHR_materials_pbrSpecularGlossiness");
 		case EGLTFJsonExtension::KHR_MaterialsIridescence:            return TEXT("KHR_materials_iridescence");
+		case EGLTFJsonExtension::KHR_MaterialsAnisotropy:             return TEXT("KHR_materials_anisotropy");
 		case EGLTFJsonExtension::KHR_MeshQuantization:                return TEXT("KHR_mesh_quantization");
 		case EGLTFJsonExtension::KHR_TextureTransform:                return TEXT("KHR_texture_transform");
 		case EGLTFJsonExtension::KHR_MaterialsSpecular:	              return TEXT("KHR_materials_specular");
diff --git a/Engine/Plugins/Enterprise/GLTFExporter/Source/GLTFExporter/Public/Json/GLTFJsonEnums.h b/Engine/Plugins/Enterprise/GLTFExporter/Source/GLTFExporter/Public/Json/GLTFJsonEnums.h
--- a/Engine/Plugins/Enterprise/GLTFExporter/Source/GLTFExporter/Public/Json/GLTFJsonEnums.h
+++ b/Engine/Plugins/Enterprise/GLTFExporter/Source/GLTFExporter/Public/Json/GLTFJsonEnums.h
@@ -14,6 +14,7 @@ enum class EGLTFJsonExtension
 	KHR_MaterialsTransmission,
 	KHR_MaterialsSpecularGlossiness,
 	KHR_MaterialsIridescence,
+	KHR_MaterialsAnisotropy,
 	KHR_MeshQuantization,
 	KHR_TextureTransform,
 	KHR_MaterialsSpecular
diff --git a/Engine/Plugins/Enterprise/GLTFExporter/Source/GLTFExporter/Public/Json/GLTFJsonMaterial.h b/Engine/Plugins/Enterprise/GLTFExporter/Source/GLTFExporter/Public/Json/GLTFJsonMaterial.h
--- a/Engine/Plugins/Enterprise/GLTFExporter/Source/GLTFExporter/Public/Json/GLTFJsonMaterial.h
+++ b/Engine/Plugins/Enterprise/GLTFExporter/Source/GLTFExporter/Public/Json/GLTFJsonMaterial.h
@@ -196,6 +196,17 @@ struct GLTFEXPORTER_API FGLTFJsonIridescenceExtension : IGLTFJsonObject
 	virtual void WriteObject(IGLTFJsonWriter& Writer) const override;
 };
 
+struct GLTFEXPORTER_API FGLTFJsonAnisotropyExtension : IGLTFJsonObject
+{
+	float                AnisotropyStrength = 0.0f;
+	float                AnisotropyRotation = 0.0f;
+	FGLTFJsonTextureInfo AnisotropyTexture;
+
+	bool HasValue() const { return !FMath::IsNearlyEqual(AnisotropyStrength, 0.f) || !FMath::IsNearlyEqual(AnisotropyRotation, 0.f) || AnisotropyTexture.Index != nullptr; }
+
+	virtual void WriteObject(IGLTFJsonWriter& Writer) const override;
+};
+
 struct GLTFEXPORTER_API FGLTFJsonMaterial : IGLTFJsonIndexedObject
 {
 	FString Name;
@@ -223,6 +234,7 @@ struct GLTFEXPORTER_API FGLTFJsonMaterial : IGLTFJsonIndexedObject
 	FGLTFJsonSheenExtension        Sheen;
 	FGLTFJsonTransmissionExtension Transmission;
 	FGLTFJsonIridescenceExtension  Iridescence;
+	FGLTFJsonAnisotropyExtension   Anisotropy;
 
 	virtual void WriteObject(IGLTFJsonWriter& Writer) const override;
 
