[glTF Exporter] Tiling method handling approach changes for glTF Proxy materials.

#rnx

#virtualized

diff --git a/Engine/Plugins/Enterprise/GLTFExporter/Source/GLTFExporter/Private/Converters/GLTFMaterialProxyFactory.cpp b/Engine/Plugins/Enterprise/GLTFExporter/Source/GLTFExporter/Private/Converters/GLTFMaterialProxyFactory.cpp
--- a/Engine/Plugins/Enterprise/GLTFExporter/Source/GLTFExporter/Private/Converters/GLTFMaterialProxyFactory.cpp
+++ b/Engine/Plugins/Enterprise/GLTFExporter/Source/GLTFExporter/Private/Converters/GLTFMaterialProxyFactory.cpp
@@ -193,6 +193,7 @@ void FGLTFMaterialProxyFactory::SetProxyParameter(UMaterialInstanceConstant* Pro
 	ParameterInfo.UVOffset.Set(ProxyMaterial, FLinearColor(TextureInfo.Transform.Offset.X, TextureInfo.Transform.Offset.Y, 0, 0), true);
 	ParameterInfo.UVScale.Set(ProxyMaterial, FLinearColor(TextureInfo.Transform.Scale.X, TextureInfo.Transform.Scale.Y, 0, 0), true);
 	ParameterInfo.UVRotation.Set(ProxyMaterial, TextureInfo.Transform.Rotation, true);
+	ParameterInfo.TilingMethod.Set(ProxyMaterial, FLinearColor(FGLTFMaterialUtilities::FromTextureWrap(TextureInfo.Index->Sampler->WrapS), FGLTFMaterialUtilities::FromTextureWrap(TextureInfo.Index->Sampler->WrapT), 0, 0), true);
 }
 
 UTexture2D* FGLTFMaterialProxyFactory::FindOrCreateTexture(FGLTFJsonTexture* JsonTexture, const FGLTFProxyMaterialTextureParameterInfo& ParameterInfo)
diff --git a/Engine/Plugins/Enterprise/GLTFExporter/Source/GLTFExporter/Private/Converters/GLTFMaterialUtilities.cpp b/Engine/Plugins/Enterprise/GLTFExporter/Source/GLTFExporter/Private/Converters/GLTFMaterialUtilities.cpp
--- a/Engine/Plugins/Enterprise/GLTFExporter/Source/GLTFExporter/Private/Converters/GLTFMaterialUtilities.cpp
+++ b/Engine/Plugins/Enterprise/GLTFExporter/Source/GLTFExporter/Private/Converters/GLTFMaterialUtilities.cpp
@@ -611,6 +611,38 @@ EMaterialShadingModel FGLTFMaterialUtilities::GetShadingModel(const UMaterialInt
 	return Possibilities.GetFirstShadingModel();
 }
 
+//To acquire TextureAddress from glTF Imported material TilingMethod values:
+TextureAddress FGLTFMaterialUtilities::ToTextureAddress(float Value)
+{
+	//0: wrap
+	//1: clamp
+	//2: mirror
+	int ValueInt = FMath::RoundToInt(Value);
+
+	switch (ValueInt)
+	{
+		case 0:     return TextureAddress::TA_Wrap;
+		case 1:     return TextureAddress::TA_Clamp;
+		case 2:	    return TextureAddress::TA_Mirror;
+		default:    return TextureAddress::TA_MAX;
+	}
+}
+float FGLTFMaterialUtilities::FromTextureWrap(EGLTFJsonTextureWrap Value)
+{
+	//0: wrap
+	//1: clamp
+	//2: mirror
+
+	switch (Value)
+	{
+		case EGLTFJsonTextureWrap::Repeat:           return 0;
+		case EGLTFJsonTextureWrap::ClampToEdge:      return 1;
+		case EGLTFJsonTextureWrap::MirroredRepeat:   return 2;
+	}
+
+	return 1;
+}
+
 FGLTFImportMaterialMatchMakingHelper::FGLTFImportMaterialMatchMakingHelper(FGLTFConvertBuilder& InBuilder,
 	const UMaterialInterface* InMaterial,
 	FGLTFJsonMaterial& InJsonMaterial)
@@ -833,23 +865,6 @@ bool FGLTFImportMaterialMatchMakingHelper::GetValue(const FString& InputKey, FGL
 #endif
 }
 
-//To acquire TextureAddress from glTF Imported material TilingMethod values:
-TextureAddress GetTextureAddress(float Value)
-{
-	//0: wrap
-	//1: clamp
-	//2: mirror
-	int ValueInt = FMath::RoundToInt(Value);
-
-	switch (ValueInt)
-	{
-		case 0:     return TextureAddress::TA_Wrap;
-		case 1:     return TextureAddress::TA_Clamp;
-		case 2:	    return TextureAddress::TA_Mirror;
-		default:    return TextureAddress::TA_MAX;
-	}
-}
-
 bool FGLTFImportMaterialMatchMakingHelper::GetValue(const FString& InputKey, FGLTFJsonTextureInfo& OutValue)
 {
 #if GLTF_EXPORT_ENABLE
@@ -892,14 +907,14 @@ bool FGLTFImportMaterialMatchMakingHelper::GetValue(const FString& InputKey, FGL
 		return false;
 	}
 
-	TextureAddress TextureAddressX = TextureAddress::TA_Wrap; //glTF Importer's default tiling method is Wrap
-	TextureAddress TextureAddressY = TextureAddress::TA_Wrap; //glTF Importer's default tiling method is Wrap
+	TextureAddress TextureAddressX = TextureAddress::TA_Wrap;
+	TextureAddress TextureAddressY = TextureAddress::TA_Wrap;
 
 	FGLTFJsonColor4 TilingMethod = FGLTFColor4{ 0.f, 0.f, 0.f, 0.f };
 	if (GetValue(InputKey + PostFix::TilingMethod, TilingMethod, false))
 	{
-		TextureAddressX = GetTextureAddress(TilingMethod.R);
-		TextureAddressY = GetTextureAddress(TilingMethod.G);
+		TextureAddressX = FGLTFMaterialUtilities::ToTextureAddress(TilingMethod.R);
+		TextureAddressY = FGLTFMaterialUtilities::ToTextureAddress(TilingMethod.G);
 	}
 
 	const bool bSRGB = InputKey == BaseColorTexture || InputKey == EmissiveTexture;
diff --git a/Engine/Plugins/Enterprise/GLTFExporter/Source/GLTFExporter/Private/Converters/GLTFMaterialUtilities.h b/Engine/Plugins/Enterprise/GLTFExporter/Source/GLTFExporter/Private/Converters/GLTFMaterialUtilities.h
--- a/Engine/Plugins/Enterprise/GLTFExporter/Source/GLTFExporter/Private/Converters/GLTFMaterialUtilities.h
+++ b/Engine/Plugins/Enterprise/GLTFExporter/Source/GLTFExporter/Private/Converters/GLTFMaterialUtilities.h
@@ -49,6 +49,9 @@ struct FGLTFMaterialUtilities
 
 	static EMaterialShadingModel GetShadingModel(const UMaterialInterface* Material, FString& OutMessage);
 
+	static TextureAddress ToTextureAddress(float Value);
+	static float FromTextureWrap(EGLTFJsonTextureWrap Value);
+
 #if WITH_EDITOR
 	static bool IsNormalMap(const FMaterialPropertyEx& Property);
 	static bool IsSRGB(const FMaterialPropertyEx& Property);
diff --git a/Engine/Plugins/Enterprise/GLTFExporter/Source/GLTFExporter/Private/Materials/GLTFProxyMaterialParameterInfo.h b/Engine/Plugins/Enterprise/GLTFExporter/Source/GLTFExporter/Private/Materials/GLTFProxyMaterialParameterInfo.h
--- a/Engine/Plugins/Enterprise/GLTFExporter/Source/GLTFExporter/Private/Materials/GLTFProxyMaterialParameterInfo.h
+++ b/Engine/Plugins/Enterprise/GLTFExporter/Source/GLTFExporter/Private/Materials/GLTFProxyMaterialParameterInfo.h
@@ -62,6 +62,7 @@ public:
 		, UVOffset(ParameterName + TEXT(" UV Offset"))
 		, UVScale(ParameterName + TEXT(" UV Scale"))
 		, UVRotation(ParameterName + TEXT(" UV Rotation"))
+		, TilingMethod(ParameterName + TEXT(" Tiling Method"))
 	{
 	}
 
@@ -80,4 +81,5 @@ public:
 	TGLTFProxyMaterialParameterInfo<FLinearColor> UVOffset;
 	TGLTFProxyMaterialParameterInfo<FLinearColor> UVScale;
 	TGLTFProxyMaterialParameterInfo<float> UVRotation;
+	TGLTFProxyMaterialParameterInfo<FLinearColor> TilingMethod;
 };
diff --git a/Engine/Plugins/Enterprise/GLTFExporter/Source/GLTFExporter/Private/Tasks/GLTFDelayedMaterialTasks.cpp b/Engine/Plugins/Enterprise/GLTFExporter/Source/GLTFExporter/Private/Tasks/GLTFDelayedMaterialTasks.cpp
--- a/Engine/Plugins/Enterprise/GLTFExporter/Source/GLTFExporter/Private/Tasks/GLTFDelayedMaterialTasks.cpp
+++ b/Engine/Plugins/Enterprise/GLTFExporter/Source/GLTFExporter/Private/Tasks/GLTFDelayedMaterialTasks.cpp
@@ -418,8 +418,18 @@ void FGLTFDelayedMaterialTask::GetProxyParameter(const FGLTFProxyMaterialTexture
 		return;
 	}
 
+	TextureAddress TextureAddressX = TextureAddress::TA_Wrap;
+	TextureAddress TextureAddressY = TextureAddress::TA_Wrap;
+
+	FLinearColor TilingMethod;
+	if (ParameterInfo.TilingMethod.Get(Material, TilingMethod, true))
+	{
+		TextureAddressX = FGLTFMaterialUtilities::ToTextureAddress(TilingMethod.R);
+		TextureAddressY = FGLTFMaterialUtilities::ToTextureAddress(TilingMethod.G);
+	}
+
 	const bool bSRGB = ParameterInfo == FGLTFProxyMaterialInfo::BaseColor || ParameterInfo == FGLTFProxyMaterialInfo::Emissive;
-	OutValue.Index = Builder.AddUniqueTexture(Texture, bSRGB);
+	OutValue.Index = Builder.AddUniqueTexture(Texture, bSRGB, TextureAddressX, TextureAddressY);
 
 	float UVIndex;
 	if (ParameterInfo.UVIndex.Get(Material, UVIndex, true))
