[glTF Exporter] Implement the ability for the user to export FMeshDescription instead of render data

#rnx

diff --git a/Engine/Plugins/Enterprise/GLTFExporter/Source/GLTFExporter/Private/Builders/GLTFBuilder.cpp b/Engine/Plugins/Enterprise/GLTFExporter/Source/GLTFExporter/Private/Builders/GLTFBuilder.cpp
--- a/Engine/Plugins/Enterprise/GLTFExporter/Source/GLTFExporter/Private/Builders/GLTFBuilder.cpp
+++ b/Engine/Plugins/Enterprise/GLTFExporter/Source/GLTFExporter/Private/Builders/GLTFBuilder.cpp
@@ -88,6 +88,13 @@ int32 FGLTFBuilder::SanitizeLOD(const USkeletalMesh* SkeletalMesh, const USkelet
 		: FGLTFMeshUtilities::GetLOD(SkeletalMesh, SkeletalMeshComponent, ExportOptions->DefaultLevelOfDetail);
 }
 
+int32 FGLTFBuilder::SanitizeLOD(const UStaticMesh* StaticMesh, const USplineMeshComponent* SplineMeshComponent, int32 LODIndex) const
+{
+	return LODIndex > 0
+		? FMath::Min(LODIndex, FGLTFMeshUtilities::GetMaximumLOD(StaticMesh))
+		: FGLTFMeshUtilities::GetLOD(StaticMesh, SplineMeshComponent, ExportOptions->DefaultLevelOfDetail);
+}
+
 const UGLTFExportOptions* FGLTFBuilder::SanitizeExportOptions(const UGLTFExportOptions* Options)
 {
 	if (Options == nullptr)
diff --git a/Engine/Plugins/Enterprise/GLTFExporter/Source/GLTFExporter/Private/Builders/GLTFConvertBuilder.cpp b/Engine/Plugins/Enterprise/GLTFExporter/Source/GLTFExporter/Private/Builders/GLTFConvertBuilder.cpp
--- a/Engine/Plugins/Enterprise/GLTFExporter/Source/GLTFExporter/Private/Builders/GLTFConvertBuilder.cpp
+++ b/Engine/Plugins/Enterprise/GLTFExporter/Source/GLTFExporter/Private/Builders/GLTFConvertBuilder.cpp
@@ -8,8 +8,7 @@
 #include "Engine/TextureRenderTarget2D.h"
 #include "LandscapeComponent.h"
 #include "Components/SplineMeshComponent.h"
-#include "Converters/GLTFNormalArray.h"
-#include "Converters/GLTFUVArray.h"
+#include "Converters/GLTFMeshAttributesArray.h"
 
 FGLTFConvertBuilder::FGLTFConvertBuilder(const FString& FileName, const UGLTFExportOptions* ExportOptions, const TSet<AActor*>& SelectedActors)
 	: FGLTFAnalyticsBuilder(FileName, ExportOptions)
@@ -38,13 +37,8 @@ FGLTFJsonAccessor* FGLTFConvertBuilder::AddUniquePositionAccessor(const FGLTFMes
 	return PositionBufferConverter->GetOrAdd(MeshSection, VertexBuffer);
 }
 
-FGLTFJsonAccessor* FGLTFConvertBuilder::AddUniquePositionAccessor(const FPositionVertexBuffer* VertexBuffer)
+FGLTFJsonAccessor* FGLTFConvertBuilder::AddUniquePositionAccessor(const FGLTFPositionArray& VertexBuffer)
 {
-	if (VertexBuffer == nullptr)
-	{
-		return nullptr;
-	}
-
 	return PositionBufferConverterRaw->GetOrAdd(VertexBuffer);
 }
 
@@ -58,6 +52,11 @@ FGLTFJsonAccessor* FGLTFConvertBuilder::AddUniqueColorAccessor(const FGLTFMeshSe
 	return ColorBufferConverter->GetOrAdd(MeshSection, VertexBuffer);
 }
 
+FGLTFJsonAccessor* FGLTFConvertBuilder::AddUniqueColorAccessor(const FGLTFColorArray& VertexColorBuffer)
+{
+	return ColorBufferConverterRaw->GetOrAdd(VertexColorBuffer);
+}
+
 FGLTFJsonAccessor* FGLTFConvertBuilder::AddUniqueNormalAccessor(const FGLTFMeshSection* MeshSection, const FStaticMeshVertexBuffer* VertexBuffer)
 {
 	if (VertexBuffer == nullptr)
@@ -68,13 +67,8 @@ FGLTFJsonAccessor* FGLTFConvertBuilder::AddUniqueNormalAccessor(const FGLTFMeshS
 	return NormalBufferConverter->GetOrAdd(MeshSection, VertexBuffer);
 }
 
-FGLTFJsonAccessor* FGLTFConvertBuilder::AddUniqueNormalAccessor(const FGLTFNormalArray* Normals)
+FGLTFJsonAccessor* FGLTFConvertBuilder::AddUniqueNormalAccessor(const FGLTFNormalArray& Normals)
 {
-	if (Normals == nullptr)
-	{
-		return nullptr;
-	}
-
 	return NormalBufferConverterRaw->GetOrAdd(Normals);
 }
 
@@ -88,6 +82,11 @@ FGLTFJsonAccessor* FGLTFConvertBuilder::AddUniqueTangentAccessor(const FGLTFMesh
 	return TangentBufferConverter->GetOrAdd(MeshSection, VertexBuffer);
 }
 
+FGLTFJsonAccessor* FGLTFConvertBuilder::AddUniqueTangentAccessor(const FGLTFTangentArray& Tangents)
+{
+	return TangentBufferConverterRaw->GetOrAdd(Tangents);
+}
+
 FGLTFJsonAccessor* FGLTFConvertBuilder::AddUniqueUVAccessor(const FGLTFMeshSection* MeshSection, const FStaticMeshVertexBuffer* VertexBuffer, int32 UVIndex)
 {
 	if (VertexBuffer == nullptr)
@@ -98,13 +97,8 @@ FGLTFJsonAccessor* FGLTFConvertBuilder::AddUniqueUVAccessor(const FGLTFMeshSecti
 	return UVBufferConverter->GetOrAdd(MeshSection, VertexBuffer, UVIndex);
 }
 
-FGLTFJsonAccessor* FGLTFConvertBuilder::AddUniqueUVAccessor(const FGLTFUVArray* UVs)
+FGLTFJsonAccessor* FGLTFConvertBuilder::AddUniqueUVAccessor(const FGLTFUVArray& UVs)
 {
-	if (UVs == nullptr)
-	{
-		return nullptr;
-	}
-
 	return UVBufferConverterRaw->GetOrAdd(UVs);
 }
 
@@ -128,6 +122,17 @@ FGLTFJsonAccessor* FGLTFConvertBuilder::AddUniqueWeightAccessor(const FGLTFMeshS
 	return BoneWeightBufferConverter->GetOrAdd(MeshSection, VertexBuffer, InfluenceOffset);
 }
 
+FGLTFJsonAccessor* FGLTFConvertBuilder::AddUniqueJointAccessor(const FGLTFJointInfluenceArray& BoneIndices)
+{
+	return BoneIndexBufferConverterRaw->GetOrAdd(BoneIndices);
+}
+
+FGLTFJsonAccessor* FGLTFConvertBuilder::AddUniqueWeightAccessor(const FGLTFJointWeightArray& Weights)
+{
+	return BoneWeightBufferConverterRaw->GetOrAdd(Weights);
+}
+
+
 FGLTFJsonAccessor* FGLTFConvertBuilder::AddUniqueIndexAccessor(const FGLTFMeshSection* MeshSection)
 {
 	if (MeshSection == nullptr)
@@ -138,13 +143,8 @@ FGLTFJsonAccessor* FGLTFConvertBuilder::AddUniqueIndexAccessor(const FGLTFMeshSe
 	return IndexBufferConverter->GetOrAdd(MeshSection);
 }
 
-FGLTFJsonAccessor* FGLTFConvertBuilder::AddUniqueIndexAccessor(const FGLTFIndexArray* IndexBuffer, const FString& MeshName)
+FGLTFJsonAccessor* FGLTFConvertBuilder::AddUniqueIndexAccessor(const FGLTFIndexArray& IndexBuffer, const FString& MeshName)
 {
-	if (IndexBuffer == nullptr)
-	{
-		return nullptr;
-	}
-
 	return IndexBufferConverterRaw->GetOrAdd(IndexBuffer, MeshName);
 }
 
diff --git a/Engine/Plugins/Enterprise/GLTFExporter/Source/GLTFExporter/Private/Converters/GLTFAccessorConverters.cpp b/Engine/Plugins/Enterprise/GLTFExporter/Source/GLTFExporter/Private/Converters/GLTFAccessorConverters.cpp
--- a/Engine/Plugins/Enterprise/GLTFExporter/Source/GLTFExporter/Private/Converters/GLTFAccessorConverters.cpp
+++ b/Engine/Plugins/Enterprise/GLTFExporter/Source/GLTFExporter/Private/Converters/GLTFAccessorConverters.cpp
@@ -8,8 +8,6 @@
 #include "Rendering/PositionVertexBuffer.h"
 #include "Rendering/SkinWeightVertexBuffer.h"
 #include "Rendering/StaticMeshVertexBuffer.h"
-#include "Converters/GLTFNormalArray.h"
-#include "Converters/GLTFUVArray.h"
 
 namespace
 {
@@ -49,118 +47,132 @@ namespace
 		}
 	};
 
-	FGLTFJsonAccessor* ConvertVertexBuffer(const FPositionVertexBuffer* VertexBuffer, FGLTFConvertBuilder& Builder, const TArray<uint32>& IndexMap = TArray<uint32>())
+	template <typename IndexType, typename IndexBufferType>
+	FGLTFJsonAccessor* ConvertIndexBuffer(const TArray<IndexBufferType>& IndexBuffer, FGLTFConvertBuilder& Builder)
 	{
-		if (VertexBuffer == nullptr || VertexBuffer->GetNumVertices() == 0)
+		const uint32 IndexCount = IndexBuffer.Num();
+		if (IndexCount == 0)
 		{
 			return nullptr;
 		}
 
-		const TUniquePtr<IGLTFBufferAdapter> SourceBuffer = IGLTFBufferAdapter::GetPositions(VertexBuffer);
-		const uint8* SourceData = SourceBuffer->GetData();
+		TArray<IndexType> Indices;
+		Indices.AddUninitialized(IndexCount);
 
-		if (SourceData == nullptr)
+		for (uint32 Index = 0; Index < IndexCount; ++Index)
 		{
-			// TODO: report error
-			return nullptr;
+			Indices[Index] = static_cast<IndexType>(IndexBuffer[Index]);
 		}
 
-		const uint32 VertexCount = IndexMap.Num() > 0 ? IndexMap.Num() : VertexBuffer->GetNumVertices();
-		const uint32 Stride = VertexBuffer->GetStride();
+		FGLTFJsonAccessor* JsonAccessor = Builder.AddAccessor();
+		JsonAccessor->BufferView = Builder.AddBufferView(Indices, EGLTFJsonBufferTarget::ElementArrayBuffer, sizeof(IndexType));
+		JsonAccessor->ComponentType = FGLTFCoreUtilities::GetComponentType<IndexType>();
+		JsonAccessor->Count = IndexCount;
+		JsonAccessor->Type = EGLTFJsonAccessorType::Scalar;
 
-		TArray<FGLTFVector3> Positions;
-		Positions.AddUninitialized(VertexCount);
+		return JsonAccessor;
+	}
 
-		if (IndexMap.Num() > 0)
-		{
-			for (uint32 VertexIndex = 0; VertexIndex < VertexCount; ++VertexIndex)
-			{
-				const uint32 MappedVertexIndex = IndexMap[VertexIndex];
-				const FVector3f& Position = *reinterpret_cast<const FVector3f*>(SourceData + Stride * MappedVertexIndex);
-				Positions[VertexIndex] = FGLTFCoreUtilities::ConvertPosition(Position, Builder.ExportOptions->ExportUniformScale);
-			}
-		}
-		else
+	template <typename ComponentType, EGLTFJsonAccessorType AccessorType, class ContainerClass>
+	FGLTFJsonAccessor* ConvertToAccessor(FGLTFConvertBuilder& Builder, ContainerClass& Container, bool bNormalized = false)
+	{
+		if (Container.Num() == 0)
 		{
-			for (uint32 VertexIndex = 0; VertexIndex < VertexCount; ++VertexIndex)
-			{
-				const FVector3f& Position = *reinterpret_cast<const FVector3f*>(SourceData + Stride * VertexIndex);
-				Positions[VertexIndex] = FGLTFCoreUtilities::ConvertPosition(Position, Builder.ExportOptions->ExportUniformScale);
-			}
+			return nullptr;
 		}
 
-
 		FGLTFJsonAccessor* JsonAccessor = Builder.AddAccessor();
-		JsonAccessor->BufferView = Builder.AddBufferView(Positions, EGLTFJsonBufferTarget::ArrayBuffer);
-		JsonAccessor->ComponentType = EGLTFJsonComponentType::Float;
-		JsonAccessor->Count = VertexCount;
-		JsonAccessor->Type = EGLTFJsonAccessorType::Vec3;
+		JsonAccessor->BufferView = Builder.AddBufferView(Container, EGLTFJsonBufferTarget::ArrayBuffer);
+		JsonAccessor->ComponentType = FGLTFCoreUtilities::GetComponentType<ComponentType>();
+		JsonAccessor->Count = Container.Num();
+		JsonAccessor->Type = AccessorType;
+		JsonAccessor->bNormalized = bNormalized;
 
-		if (VertexCount > 0)
-		{
-			// Calculate accurate bounding box based on raw vertex values
-			JsonAccessor->MinMaxLength = 3;
+		return JsonAccessor;
+	}
+}
 
-			for (int32 ComponentIndex = 0; ComponentIndex < JsonAccessor->MinMaxLength; ComponentIndex++)
-			{
-				JsonAccessor->Min[ComponentIndex] = Positions[0].Components[ComponentIndex];
-				JsonAccessor->Max[ComponentIndex] = Positions[0].Components[ComponentIndex];
-			}
+FGLTFJsonAccessor* FGLTFPositionBufferConverter::Convert(const FGLTFMeshSection* MeshSection, const FPositionVertexBuffer* VertexBuffer)
+{
+	if (VertexBuffer == nullptr || VertexBuffer->GetNumVertices() == 0)
+	{
+		return nullptr;
+	}
 
-			for (uint32 VertexIndex = 1; VertexIndex < VertexCount; ++VertexIndex)
-			{
-				const FGLTFVector3& Position = Positions[VertexIndex];
-				for (int32 ComponentIndex = 0; ComponentIndex < JsonAccessor->MinMaxLength; ComponentIndex++)
-				{
-					JsonAccessor->Min[ComponentIndex] = FMath::Min(JsonAccessor->Min[ComponentIndex], Position.Components[ComponentIndex]);
-					JsonAccessor->Max[ComponentIndex] = FMath::Max(JsonAccessor->Max[ComponentIndex], Position.Components[ComponentIndex]);
-				}
-			}
-		}
+	const TUniquePtr<IGLTFBufferAdapter> SourceBuffer = IGLTFBufferAdapter::GetPositions(VertexBuffer);
+	const uint8* SourceData = SourceBuffer->GetData();
 
-		return JsonAccessor;
+	if (SourceData == nullptr)
+	{
+		// TODO: report error
+		return nullptr;
 	}
 
-	template <typename IndexType, typename IndexBufferType>
-	FGLTFJsonAccessor* ConvertIndexBuffer(const TArray<IndexBufferType>& IndexBuffer, FGLTFConvertBuilder& Builder)
+	const TArray<uint32>& IndexMap = MeshSection->IndexMap;
+	const uint32 VertexCount = IndexMap.Num() > 0 ? IndexMap.Num() : VertexBuffer->GetNumVertices();
+	const uint32 Stride = VertexBuffer->GetStride();
+
+	TArray<FGLTFVector3> Positions;
+	Positions.AddUninitialized(VertexCount);
+
+	if (IndexMap.Num() > 0)
 	{
-		const uint32 IndexCount = IndexBuffer.Num();
-		if (IndexCount == 0)
+		for (uint32 VertexIndex = 0; VertexIndex < VertexCount; ++VertexIndex)
 		{
-			return nullptr;
+			const uint32 MappedVertexIndex = IndexMap[VertexIndex];
+			const FVector3f& Position = *reinterpret_cast<const FVector3f*>(SourceData + Stride * MappedVertexIndex);
+			Positions[VertexIndex] = FGLTFCoreUtilities::ConvertPosition(Position, Builder.ExportOptions->ExportUniformScale);
+		}
+	}
+	else
+	{
+		for (uint32 VertexIndex = 0; VertexIndex < VertexCount; ++VertexIndex)
+		{
+			const FVector3f& Position = *reinterpret_cast<const FVector3f*>(SourceData + Stride * VertexIndex);
+			Positions[VertexIndex] = FGLTFCoreUtilities::ConvertPosition(Position, Builder.ExportOptions->ExportUniformScale);
 		}
+	}
 
-		TArray<IndexType> Indices;
-		Indices.AddUninitialized(IndexCount);
 
-		for (uint32 Index = 0; Index < IndexCount; ++Index)
+	FGLTFJsonAccessor* JsonAccessor = Builder.AddAccessor();
+	JsonAccessor->BufferView = Builder.AddBufferView(Positions, EGLTFJsonBufferTarget::ArrayBuffer);
+	JsonAccessor->ComponentType = EGLTFJsonComponentType::Float;
+	JsonAccessor->Count = VertexCount;
+	JsonAccessor->Type = EGLTFJsonAccessorType::Vec3;
+
+	if (VertexCount > 0)
+	{
+		// Calculate accurate bounding box based on raw vertex values
+		JsonAccessor->MinMaxLength = 3;
+
+		for (int32 ComponentIndex = 0; ComponentIndex < JsonAccessor->MinMaxLength; ComponentIndex++)
 		{
-			Indices[Index] = static_cast<IndexType>(IndexBuffer[Index]);
+			JsonAccessor->Min[ComponentIndex] = Positions[0].Components[ComponentIndex];
+			JsonAccessor->Max[ComponentIndex] = Positions[0].Components[ComponentIndex];
 		}
 
-		FGLTFJsonAccessor* JsonAccessor = Builder.AddAccessor();
-		JsonAccessor->BufferView = Builder.AddBufferView(Indices, EGLTFJsonBufferTarget::ElementArrayBuffer, sizeof(IndexType));
-		JsonAccessor->ComponentType = FGLTFCoreUtilities::GetComponentType<IndexType>();
-		JsonAccessor->Count = IndexCount;
-		JsonAccessor->Type = EGLTFJsonAccessorType::Scalar;
-
-		return JsonAccessor;
+		for (uint32 VertexIndex = 1; VertexIndex < VertexCount; ++VertexIndex)
+		{
+			const FGLTFVector3& Position = Positions[VertexIndex];
+			for (int32 ComponentIndex = 0; ComponentIndex < JsonAccessor->MinMaxLength; ComponentIndex++)
+			{
+				JsonAccessor->Min[ComponentIndex] = FMath::Min(JsonAccessor->Min[ComponentIndex], Position.Components[ComponentIndex]);
+				JsonAccessor->Max[ComponentIndex] = FMath::Max(JsonAccessor->Max[ComponentIndex], Position.Components[ComponentIndex]);
+			}
+		}
 	}
-}
 
-FGLTFJsonAccessor* FGLTFPositionBufferConverter::Convert(const FGLTFMeshSection* MeshSection, const FPositionVertexBuffer* VertexBuffer)
-{
-	return ConvertVertexBuffer(VertexBuffer, Builder, MeshSection->IndexMap);
+	return JsonAccessor;
 }
 
-FGLTFJsonAccessor* FGLTFColorBufferConverter::Convert(const FGLTFMeshSection* MeshSection, const FColorVertexBuffer* VertexBuffer)
+FGLTFJsonAccessor* FGLTFColorBufferConverter::Convert(const FGLTFMeshSection* MeshSection, const FColorVertexBuffer* VertexColorBufferBuffer)
 {
-	if (VertexBuffer == nullptr || VertexBuffer->GetNumVertices() == 0)
+	if (VertexColorBufferBuffer == nullptr || VertexColorBufferBuffer->GetNumVertices() == 0)
 	{
 		return nullptr;
 	}
 
-	const TUniquePtr<IGLTFBufferAdapter> SourceBuffer = IGLTFBufferAdapter::GetColors(VertexBuffer);
+	const TUniquePtr<IGLTFBufferAdapter> SourceBuffer = IGLTFBufferAdapter::GetColors(VertexColorBufferBuffer);
 	const uint8* SourceData = SourceBuffer->GetData();
 
 	if (SourceData == nullptr)
@@ -170,18 +182,31 @@ FGLTFJsonAccessor* FGLTFColorBufferConverter::Convert(const FGLTFMeshSection* Me
 	}
 
 	const TArray<uint32>& IndexMap = MeshSection->IndexMap;
-	const uint32 VertexCount = IndexMap.Num();
-	const uint32 Stride = VertexBuffer->GetStride();
+	const uint32 VertexCount = IndexMap.Num() > 0 ? IndexMap.Num() : VertexColorBufferBuffer->GetNumVertices();
+	const uint32 Stride = VertexColorBufferBuffer->GetStride();
 
 	TArray<FGLTFUInt8Color4> Colors;
 	Colors.AddUninitialized(VertexCount);
 
-	for (uint32 VertexIndex = 0; VertexIndex < VertexCount; ++VertexIndex)
+	if (IndexMap.Num() > 0)
 	{
-		const uint32 MappedVertexIndex = IndexMap[VertexIndex];
-		const FColor& Color = *reinterpret_cast<const FColor*>(SourceData + Stride * MappedVertexIndex);
-		Colors[VertexIndex] = FGLTFCoreUtilities::ConvertColor(Color);
+		for (uint32 VertexIndex = 0; VertexIndex < VertexCount; ++VertexIndex)
+		{
+			const uint32 MappedVertexIndex = IndexMap[VertexIndex];
+			const FColor& Color = *reinterpret_cast<const FColor*>(SourceData + Stride * MappedVertexIndex);
+			Colors[VertexIndex] = FGLTFCoreUtilities::ConvertColor(Color);
+		}
 	}
+	else
+	{
+		for (uint32 VertexIndex = 0; VertexIndex < VertexCount; ++VertexIndex)
+		{
+			const FColor& Color = *reinterpret_cast<const FColor*>(SourceData + Stride * VertexIndex);
+			Colors[VertexIndex] = FGLTFCoreUtilities::ConvertColor(Color);
+		}
+	}
+
+
 
 	FGLTFJsonAccessor* JsonAccessor = Builder.AddAccessor();
 	JsonAccessor->BufferView = Builder.AddBufferView(Colors, EGLTFJsonBufferTarget::ArrayBuffer);
@@ -398,13 +423,7 @@ FGLTFJsonAccessor* FGLTFUVBufferConverter::Convert(const FGLTFMeshSection* MeshS
 		UVs[VertexIndex] = FGLTFCoreUtilities::ConvertUV(UV);
 	}
 
-	FGLTFJsonAccessor* JsonAccessor = Builder.AddAccessor();
-	JsonAccessor->BufferView = Builder.AddBufferView(UVs, EGLTFJsonBufferTarget::ArrayBuffer);
-	JsonAccessor->ComponentType = EGLTFJsonComponentType::Float;
-	JsonAccessor->Count = VertexCount;
-	JsonAccessor->Type = EGLTFJsonAccessorType::Vec2;
-
-	return JsonAccessor;
+	return ConvertToAccessor<float, EGLTFJsonAccessorType::Vec2>(Builder, UVs);
 }
 
 FGLTFJsonAccessor* FGLTFBoneIndexBufferConverter::Convert(const FGLTFMeshSection* MeshSection, const FSkinWeightVertexBuffer* VertexBuffer, uint32 InfluenceOffset)
@@ -505,13 +524,7 @@ FGLTFJsonAccessor* FGLTFBoneIndexBufferConverter::Convert(const FGLTFMeshSection
 		}
 	}
 
-	FGLTFJsonAccessor* JsonAccessor = Builder.AddAccessor();
-	JsonAccessor->BufferView = Builder.AddBufferView(BoneIndices, EGLTFJsonBufferTarget::ArrayBuffer);
-	JsonAccessor->ComponentType = FGLTFCoreUtilities::GetComponentType<DestinationType>();
-	JsonAccessor->Count = VertexCount;
-	JsonAccessor->Type = EGLTFJsonAccessorType::Vec4;
-
-	return JsonAccessor;
+	return ConvertToAccessor<DestinationType, EGLTFJsonAccessorType::Vec4>(Builder, BoneIndices);
 }
 
 FGLTFJsonAccessor* FGLTFBoneWeightBufferConverter::Convert(const FGLTFMeshSection* MeshSection, const FSkinWeightVertexBuffer* VertexBuffer, uint32 InfluenceOffset)
@@ -602,14 +615,7 @@ FGLTFJsonAccessor* FGLTFBoneWeightBufferConverter::Convert(const FGLTFMeshSectio
 		}
 	}
 
-	FGLTFJsonAccessor* JsonAccessor = Builder.AddAccessor();
-	JsonAccessor->BufferView = Builder.AddBufferView(BoneWeights, EGLTFJsonBufferTarget::ArrayBuffer);
-	JsonAccessor->ComponentType = EGLTFJsonComponentType::UInt8;
-	JsonAccessor->Count = VertexCount;
-	JsonAccessor->Type = EGLTFJsonAccessorType::Vec4;
-	JsonAccessor->bNormalized = true;
-
-	return JsonAccessor;
+	return ConvertToAccessor<uint8, EGLTFJsonAccessorType::Vec4>(Builder, BoneWeights, true);
 }
 
 FGLTFJsonAccessor* FGLTFIndexBufferConverter::Convert(const FGLTFMeshSection* MeshSection)
@@ -632,18 +638,57 @@ FGLTFJsonAccessor* FGLTFIndexBufferConverter::Convert(const FGLTFMeshSection* Me
 }
 
 
-FGLTFJsonAccessor* FGLTFPositionBufferConverterRaw::Convert(const FPositionVertexBuffer* VertexBuffer)
+FGLTFJsonAccessor* FGLTFPositionBufferConverterRaw::Convert(FGLTFPositionArray VertexBuffer)
 {
-	return ConvertVertexBuffer(VertexBuffer, Builder);
+	const uint32 VertexCount = VertexBuffer.Num();
+
+	TArray<FGLTFVector3> Positions;
+	FVector3f Min, Max;
+
+	Positions.AddUninitialized(VertexCount);
+
+	const size_t MinMaxLength = 3;
+	if (VertexCount > 0)
+	{
+		Positions[0] = FGLTFCoreUtilities::ConvertPosition(VertexBuffer[0], Builder.ExportOptions->ExportUniformScale);
+
+		for (size_t ComponentIndex = 0; ComponentIndex < MinMaxLength; ComponentIndex++)
+		{
+			Min[ComponentIndex] = Positions[0].Components[ComponentIndex];
+			Max[ComponentIndex] = Positions[0].Components[ComponentIndex];
+		}
+
+		for (uint32 PositionIndex = 1; PositionIndex < VertexCount; ++PositionIndex)
+		{
+			Positions[PositionIndex] = FGLTFCoreUtilities::ConvertPosition(VertexBuffer[PositionIndex], Builder.ExportOptions->ExportUniformScale);
+
+			for (size_t ComponentIndex = 0; ComponentIndex < MinMaxLength; ComponentIndex++)
+			{
+				Min[ComponentIndex] = FMath::Min(Min[ComponentIndex], Positions[PositionIndex].Components[ComponentIndex]);
+				Max[ComponentIndex] = FMath::Max(Max[ComponentIndex], Positions[PositionIndex].Components[ComponentIndex]);
+			}
+		}
+	}
+
+	FGLTFJsonAccessor* JsonAccessor = ConvertToAccessor<float, EGLTFJsonAccessorType::Vec3>(Builder, Positions);
+
+	JsonAccessor->MinMaxLength = MinMaxLength;
+	for (int32 ComponentIndex = 0; ComponentIndex < JsonAccessor->MinMaxLength; ComponentIndex++)
+	{
+		JsonAccessor->Min[ComponentIndex] = Min[ComponentIndex];
+		JsonAccessor->Max[ComponentIndex] = Max[ComponentIndex];
+	}
+
+	return JsonAccessor;
 }
 
-FGLTFJsonAccessor* FGLTFIndexBufferConverterRaw::Convert(const FGLTFIndexArray* IndexBuffer, FString MeshName)
+FGLTFJsonAccessor* FGLTFIndexBufferConverterRaw::Convert(FGLTFIndexArray IndexBuffer, FString MeshName)
 {
-	const uint32 MaxVertexIndex = IndexBuffer->Num() - 1;
+	const uint32 MaxVertexIndex = IndexBuffer.Num() - 1;
 	// NOTE: Even if a maximum value (i.e. UINT8_MAX or UINT16_MAX) would fit inside the type (i.e. uint8 or uint16), it's not allowed per gltf spec
 	// since the maximum values trigger primitive restart in some graphics APIs and would require client implementations to rebuild the index buffer.
-	if (MaxVertexIndex < UINT8_MAX) return ConvertIndexBuffer<uint8>(*IndexBuffer, Builder);
-	if (MaxVertexIndex < UINT16_MAX) return ConvertIndexBuffer<uint16>(*IndexBuffer, Builder);
+	if (MaxVertexIndex < UINT8_MAX) return ConvertIndexBuffer<uint8>(IndexBuffer, Builder);
+	if (MaxVertexIndex < UINT16_MAX) return ConvertIndexBuffer<uint16>(IndexBuffer, Builder);
 	if (MaxVertexIndex == UINT32_MAX)
 	{
 		Builder.LogWarning(
@@ -652,68 +697,76 @@ FGLTFJsonAccessor* FGLTFIndexBufferConverterRaw::Convert(const FGLTFIndexArray*
 				*MeshName
 			));
 	}
-	return ConvertIndexBuffer<uint32>(*IndexBuffer, Builder);
+	return ConvertIndexBuffer<uint32>(IndexBuffer, Builder);
 }
 
 
-FGLTFJsonAccessor* FGLTFNormalBufferConverterRaw::Convert(const FGLTFNormalArray* NormalsSource)
+FGLTFJsonAccessor* FGLTFNormalBufferConverterRaw::Convert(FGLTFNormalArray NormalsSource)
 {
-	if (NormalsSource->Num() == 0)
+	const uint32 VertexCount = NormalsSource.Num();
+
+	TArray<FGLTFVector3> Normals;
+	Normals.AddUninitialized(VertexCount);
+
+	for (uint32 VertexIndex = 0; VertexIndex < VertexCount; ++VertexIndex)
 	{
-		return nullptr;
+		Normals[VertexIndex] = FGLTFCoreUtilities::ConvertNormal(NormalsSource[VertexIndex]);
 	}
 
-	FGLTFJsonAccessor* JsonAccessor = Builder.AddAccessor();
-	JsonAccessor->BufferView = ConvertBufferView(NormalsSource);;
-	JsonAccessor->ComponentType = EGLTFJsonComponentType::Float;
-	JsonAccessor->Count = NormalsSource->Num();
-	JsonAccessor->Type = EGLTFJsonAccessorType::Vec3;
-
-	return JsonAccessor;
+	return ConvertToAccessor<float, EGLTFJsonAccessorType::Vec3>(Builder, Normals);
 }
 
-FGLTFJsonBufferView* FGLTFNormalBufferConverterRaw::ConvertBufferView(const FGLTFNormalArray* NormalsSource)
+FGLTFJsonAccessor* FGLTFUVBufferConverterRaw::Convert(FGLTFUVArray UVSource)
 {
-	const uint32 VertexCount = NormalsSource->Num();
+	const uint32 VertexCount = UVSource.Num();
 
-	TArray<FGLTFVector3> Normals;
-	Normals.AddUninitialized(VertexCount);
+	TArray<FGLTFVector2> UVs;
+	UVs.AddUninitialized(VertexCount);
 
 	for (uint32 VertexIndex = 0; VertexIndex < VertexCount; ++VertexIndex)
 	{
-		Normals[VertexIndex] = FGLTFCoreUtilities::ConvertNormal((*NormalsSource)[VertexIndex]);
+		UVs[VertexIndex] = FGLTFCoreUtilities::ConvertUV(UVSource[VertexIndex]);
 	}
 
-	return Builder.AddBufferView(Normals, EGLTFJsonBufferTarget::ArrayBuffer);
+	return ConvertToAccessor<float, EGLTFJsonAccessorType::Vec2>(Builder, UVs);
 }
 
-FGLTFJsonAccessor* FGLTFUVBufferConverterRaw::Convert(const FGLTFUVArray* UVSource)
+FGLTFJsonAccessor* FGLTFColorBufferConverterRaw::Convert(FGLTFColorArray VertexColorBuffer)
 {
-	if (UVSource->Num() == 0)
+	const uint32 VertexCount = VertexColorBuffer.Num();
+
+	TArray<FGLTFUInt8Color4> Colors;
+	Colors.AddUninitialized(VertexCount);
+
+	for (uint32 VertexIndex = 0; VertexIndex < VertexCount; ++VertexIndex)
 	{
-		return nullptr;
+		Colors[VertexIndex] = FGLTFCoreUtilities::ConvertColor(VertexColorBuffer[VertexIndex]);
 	}
 
-	FGLTFJsonAccessor* JsonAccessor = Builder.AddAccessor();
-	JsonAccessor->BufferView = ConvertBufferView(UVSource);;
-	JsonAccessor->ComponentType = EGLTFJsonComponentType::Float;
-	JsonAccessor->Count = UVSource->Num();
-	JsonAccessor->Type = EGLTFJsonAccessorType::Vec2;
-
-	return JsonAccessor;
+	return ConvertToAccessor<uint8, EGLTFJsonAccessorType::Vec4>(Builder, Colors, true);
 }
 
-FGLTFJsonBufferView* FGLTFUVBufferConverterRaw::ConvertBufferView(const FGLTFUVArray* UVSource)
+FGLTFJsonAccessor* FGLTFTangentBufferConverterRaw::Convert(FGLTFTangentArray TangentsSource)
 {
-	const uint32 VertexCount = UVSource->Num();
+	const uint32 VertexCount = TangentsSource.Num();
 
-	TArray<FGLTFVector2> UVs;
-	UVs.AddUninitialized(VertexCount);
+	TArray<FGLTFVector4> Tangents;
+	Tangents.AddUninitialized(VertexCount);
 
 	for (uint32 VertexIndex = 0; VertexIndex < VertexCount; ++VertexIndex)
 	{
-		UVs[VertexIndex] = FGLTFCoreUtilities::ConvertUV((*UVSource)[VertexIndex]);
+		Tangents[VertexIndex] = FGLTFCoreUtilities::ConvertTangent(TangentsSource[VertexIndex]);
 	}
+	
+	return ConvertToAccessor<float, EGLTFJsonAccessorType::Vec4>(Builder, Tangents);
+}
 
-	return Builder.AddBufferView(UVs, EGLTFJsonBufferTarget::ArrayBuffer);
+FGLTFJsonAccessor* FGLTFBoneIndexBufferConverterRaw::Convert(FGLTFJointInfluenceArray BoneIndices)
+{
+	return ConvertToAccessor<FBoneIndexType, EGLTFJsonAccessorType::Vec4>(Builder, BoneIndices);
+}
+
+FGLTFJsonAccessor* FGLTFBoneWeightBufferConverterRaw::Convert(FGLTFJointWeightArray Weights)
+{
+	return ConvertToAccessor<uint16, EGLTFJsonAccessorType::Vec4>(Builder, Weights, true);
 }
\ No newline at end of file
diff --git a/Engine/Plugins/Enterprise/GLTFExporter/Source/GLTFExporter/Private/Converters/GLTFMeshConverters.cpp b/Engine/Plugins/Enterprise/GLTFExporter/Source/GLTFExporter/Private/Converters/GLTFMeshConverters.cpp
--- a/Engine/Plugins/Enterprise/GLTFExporter/Source/GLTFExporter/Private/Converters/GLTFMeshConverters.cpp
+++ b/Engine/Plugins/Enterprise/GLTFExporter/Source/GLTFExporter/Private/Converters/GLTFMeshConverters.cpp
@@ -105,6 +105,8 @@ void FGLTFSplineMeshConverter::Sanitize(const UStaticMesh*& StaticMesh, const US
 
 	FGLTFMeshUtilities::ResolveMaterials(Materials, SplineMeshComponent, StaticMesh);
 	Builder.ResolveProxies(Materials);
+
+	LODIndex = Builder.SanitizeLOD(StaticMesh, SplineMeshComponent, LODIndex);
 }
 
 FGLTFJsonMesh* FGLTFSplineMeshConverter::Convert(const UStaticMesh* StaticMesh, const USplineMeshComponent* SplineMeshComponent, FGLTFMaterialArray Materials, int32 LODIndex)
@@ -124,7 +126,7 @@ FGLTFJsonMesh* FGLTFSplineMeshConverter::Convert(const UStaticMesh* StaticMesh,
 
 	if (StaticMesh && SplineMeshComponent)
 	{
-		Builder.ScheduleSlowTask<FGLTFDelayedSplineMeshTask>(Builder, /*PositionBuffer,*/ MeshSectionConverter, *StaticMesh, *SplineMeshComponent, Materials, LODIndex, JsonMesh);
+		Builder.ScheduleSlowTask<FGLTFDelayedSplineMeshTask>(Builder, /*PositionBuffer,*/ MeshSectionConverter, StaticMesh, SplineMeshComponent, Materials, LODIndex, JsonMesh);
 	}
 	
 	return JsonMesh;
diff --git a/Engine/Plugins/Enterprise/GLTFExporter/Source/GLTFExporter/Private/Converters/GLTFMeshUtilities.cpp b/Engine/Plugins/Enterprise/GLTFExporter/Source/GLTFExporter/Private/Converters/GLTFMeshUtilities.cpp
--- a/Engine/Plugins/Enterprise/GLTFExporter/Source/GLTFExporter/Private/Converters/GLTFMeshUtilities.cpp
+++ b/Engine/Plugins/Enterprise/GLTFExporter/Source/GLTFExporter/Private/Converters/GLTFMeshUtilities.cpp
@@ -11,6 +11,8 @@
 #include "StaticMeshCompiler.h"
 #include "SkinnedAssetCompiler.h"
 #include "LandscapeComponent.h"
+#include "Components/SplineMeshComponent.h"
+#include "Converters/GLTFMeshAttributesArray.h"
 
 void FGLTFMeshUtilities::FullyLoad(const UStaticMesh* InStaticMesh)
 {
@@ -214,6 +216,13 @@ int32 FGLTFMeshUtilities::GetLOD(const USkeletalMesh* SkeletalMesh, const USkele
 	return FMath::Min(LOD, GetMaximumLOD(SkeletalMesh));
 }
 
+int32 FGLTFMeshUtilities::GetLOD(const UStaticMesh* StaticMesh, const USplineMeshComponent* SplineMeshComponent, int32 DefaultLOD)
+{
+	const int32 ForcedLOD = SplineMeshComponent != nullptr ? SplineMeshComponent->ForcedLodModel - 1 : -1;
+	const int32 LOD = ForcedLOD > 0 ? ForcedLOD : FMath::Max(DefaultLOD, GetMinimumLOD(StaticMesh, SplineMeshComponent));
+	return FMath::Min(LOD, GetMaximumLOD(StaticMesh));
+}
+
 int32 FGLTFMeshUtilities::GetMaximumLOD(const UStaticMesh* StaticMesh)
 {
 	return StaticMesh != nullptr ? StaticMesh->GetNumLODs() - 1 : -1;
diff --git a/Engine/Plugins/Enterprise/GLTFExporter/Source/GLTFExporter/Private/Converters/GLTFMeshUtilities.h b/Engine/Plugins/Enterprise/GLTFExporter/Source/GLTFExporter/Private/Converters/GLTFMeshUtilities.h
--- a/Engine/Plugins/Enterprise/GLTFExporter/Source/GLTFExporter/Private/Converters/GLTFMeshUtilities.h
+++ b/Engine/Plugins/Enterprise/GLTFExporter/Source/GLTFExporter/Private/Converters/GLTFMeshUtilities.h
@@ -2,7 +2,7 @@
 
 #pragma once
 
-#include "Converters/GLTFIndexArray.h"
+#include "Converters/GLTFMeshAttributesArray.h"
 
 class UMaterialInterface;
 class UStaticMesh;
@@ -14,6 +14,7 @@ class ULandscapeComponent;
 struct FStaticMaterial;
 struct FSkeletalMaterial;
 struct FStaticMeshLODResources;
+class USplineMeshComponent;
 
 struct FGLTFMeshUtilities
 {
@@ -48,6 +49,7 @@ struct FGLTFMeshUtilities
 
 	static int32 GetLOD(const UStaticMesh* StaticMesh, const UStaticMeshComponent* StaticMeshComponent, int32 DefaultLOD);
 	static int32 GetLOD(const USkeletalMesh* SkeletalMesh, const USkeletalMeshComponent* SkeletalMeshComponent, int32 DefaultLOD);
+	static int32 GetLOD(const UStaticMesh* StaticMesh, const USplineMeshComponent* SplineMeshComponent, int32 DefaultLOD);
 
 	static int32 GetMaximumLOD(const UStaticMesh* StaticMesh);
 	static int32 GetMaximumLOD(const USkeletalMesh* SkeletalMesh);
diff --git a/Engine/Plugins/Enterprise/GLTFExporter/Source/GLTFExporter/Private/Json/GLTFJsonMesh.cpp b/Engine/Plugins/Enterprise/GLTFExporter/Source/GLTFExporter/Private/Json/GLTFJsonMesh.cpp
--- a/Engine/Plugins/Enterprise/GLTFExporter/Source/GLTFExporter/Private/Json/GLTFJsonMesh.cpp
+++ b/Engine/Plugins/Enterprise/GLTFExporter/Source/GLTFExporter/Private/Json/GLTFJsonMesh.cpp
@@ -30,9 +30,28 @@ void FGLTFJsonAttributes::WriteObject(IGLTFJsonWriter& Writer) const
 	}
 }
 
+bool FGLTFJsonAttributes::HasValue() const
+{
+	return (Position && Position->Count)
+		|| (Color0 && Color0->Count)
+		|| (Normal && Normal->Count)
+		|| (Tangent && Tangent->Count)
+		|| TexCoords.Num()
+		|| Joints.Num()
+		|| Weights.Num();
+}
+
 void FGLTFJsonPrimitive::WriteObject(IGLTFJsonWriter& Writer) const
 {
-	Writer.Write(TEXT("attributes"), Attributes);
+	if (!(Attributes.HasValue() && Indices != nullptr))
+	{
+		return;
+	}
+
+	if (Attributes.HasValue())
+	{
+		Writer.Write(TEXT("attributes"), Attributes);
+	}
 
 	if (Indices != nullptr)
 	{
@@ -61,6 +80,11 @@ void FGLTFJsonPrimitive::WriteObject(IGLTFJsonWriter& Writer) const
 	}
 }
 
+bool FGLTFJsonPrimitive::HasValue() const
+{
+	return (Attributes.HasValue() && Indices != nullptr);
+}
+
 void FGLTFJsonMesh::WriteObject(IGLTFJsonWriter& Writer) const
 {
 	if (!Name.IsEmpty())
@@ -68,5 +92,29 @@ void FGLTFJsonMesh::WriteObject(IGLTFJsonWriter& Writer) const
 		Writer.Write(TEXT("name"), Name);
 	}
 
-	Writer.Write(TEXT("primitives"), Primitives);
+	if (HasValue())
+	{
+		Writer.SetIdentifier(TEXT("primitives"));
+		Writer.StartArray();
+		for (const FGLTFJsonPrimitive& Primitive : Primitives)
+		{
+			if (Primitive.HasValue())
+			{
+				Writer.Write(Primitive);
+			}
+		}
+		Writer.EndArray();
+	}
 }
+
+bool FGLTFJsonMesh::HasValue() const
+{
+	for (const FGLTFJsonPrimitive& Primitive : Primitives)
+	{
+		if (Primitive.HasValue())
+		{
+			return true;
+		}
+	}
+	return false;
+}
\ No newline at end of file
diff --git a/Engine/Plugins/Enterprise/GLTFExporter/Source/GLTFExporter/Private/Options/GLTFExportOptions.cpp b/Engine/Plugins/Enterprise/GLTFExporter/Source/GLTFExporter/Private/Options/GLTFExportOptions.cpp
--- a/Engine/Plugins/Enterprise/GLTFExporter/Source/GLTFExporter/Private/Options/GLTFExportOptions.cpp
+++ b/Engine/Plugins/Enterprise/GLTFExporter/Source/GLTFExporter/Private/Options/GLTFExportOptions.cpp
@@ -27,6 +27,7 @@ void UGLTFExportOptions::ResetToDefault()
 	DefaultMaterialBakeFilter = TF_Trilinear;
 	DefaultMaterialBakeTiling = TA_Wrap;
 	DefaultLevelOfDetail = 0;
+	bExportSourceModel = false;
 	bExportVertexColors = false;
 	bExportVertexSkinWeights = true;
 	bMakeSkinnedMeshesRoot = true;
diff --git a/Engine/Plugins/Enterprise/GLTFExporter/Source/GLTFExporter/Private/Tasks/GLTFDelayedMaterialTasks.h b/Engine/Plugins/Enterprise/GLTFExporter/Source/GLTFExporter/Private/Tasks/GLTFDelayedMaterialTasks.h
--- a/Engine/Plugins/Enterprise/GLTFExporter/Source/GLTFExporter/Private/Tasks/GLTFDelayedMaterialTasks.h
+++ b/Engine/Plugins/Enterprise/GLTFExporter/Source/GLTFExporter/Private/Tasks/GLTFDelayedMaterialTasks.h
@@ -19,7 +19,7 @@ class FGLTFDelayedMaterialTask : public FGLTFDelayedTask
 {
 public:
 
-	FGLTFDelayedMaterialTask(FGLTFConvertBuilder& Builder, FGLTFUVOverlapChecker& UVOverlapChecker, FGLTFUVBoundsCalculator& UVBoundsCalculator, const UMaterialInterface* Material, const FGLTFMeshData* MeshData, FGLTFIndexArray SectionIndices, FGLTFJsonMaterial* JsonMaterial)
+	FGLTFDelayedMaterialTask(FGLTFConvertBuilder& Builder, FGLTFUVOverlapChecker& UVOverlapChecker, FGLTFUVBoundsCalculator& UVBoundsCalculator, const UMaterialInterface* Material, const FGLTFMeshData* MeshData, FGLTFIndexArray& SectionIndices, FGLTFJsonMaterial* JsonMaterial)
 		: FGLTFDelayedTask(EGLTFTaskPriority::Material)
 		, Builder(Builder)
 		, UVOverlapChecker(UVOverlapChecker)
@@ -45,7 +45,7 @@ private:
 	FGLTFUVBoundsCalculator& UVBoundsCalculator;
 	const UMaterialInterface* Material;
 	const FGLTFMeshData* MeshData;
-	const FGLTFIndexArray SectionIndices;
+	const FGLTFIndexArray SectionIndices; //We copy because the received value can be temporary value.
 	FGLTFJsonMaterial* JsonMaterial;
 
 	FString GetMaterialName() const;
diff --git a/Engine/Plugins/Enterprise/GLTFExporter/Source/GLTFExporter/Private/Tasks/GLTFDelayedMeshTasks.cpp b/Engine/Plugins/Enterprise/GLTFExporter/Source/GLTFExporter/Private/Tasks/GLTFDelayedMeshTasks.cpp
--- a/Engine/Plugins/Enterprise/GLTFExporter/Source/GLTFExporter/Private/Tasks/GLTFDelayedMeshTasks.cpp
+++ b/Engine/Plugins/Enterprise/GLTFExporter/Source/GLTFExporter/Private/Tasks/GLTFDelayedMeshTasks.cpp
@@ -12,6 +12,7 @@
 #include "StaticMeshComponentLODInfo.h"
 #include "StaticMeshResources.h"
 #include "Rendering/SkeletalMeshRenderData.h"
+#include "Rendering/SkeletalMeshModel.h"
 
 #include "Converters/GLTFMaterialUtilities.h"
 #include "Rendering/PositionVertexBuffer.h"
@@ -19,8 +20,7 @@
 #include "LandscapeProxy.h"
 #include "LandscapeComponent.h"
 #include "LandscapeLayerInfoObject.h"
-#include "Converters/GLTFNormalArray.h"
-#include "Converters/GLTFUVArray.h"
+#include "Converters/GLTFMeshAttributesArray.h"
 #include "Utilities/GLTFLandscapeComponentDataInterface.h"
 
 namespace
@@ -117,19 +117,305 @@ namespace
 
 		return false;
 	}
+
+	template <class T>
+	bool DoesBufferHasZeroVector(TArray<T> Buffer, float Tolerance = UE_KINDA_SMALL_NUMBER)
+	{
+		for (const T& Value : Buffer)
+		{
+			if (FMath::Abs(Value.X) <= Tolerance
+				&& FMath::Abs(Value.Y) <= Tolerance
+				&& FMath::Abs(Value.Z) <= Tolerance)
+			{
+				return true;
+			}
+		}
+		return false;
+	}
 }
 
-FString FGLTFDelayedStaticMeshTask::GetName()
+FString FGLTFDelayedStaticAndSplineMeshTask::GetName()
 {
-	return StaticMeshComponent != nullptr ? FGLTFNameUtilities::GetName(StaticMeshComponent) : StaticMesh->GetName();
+	return StaticMeshComponent != nullptr ? FGLTFNameUtilities::GetName(StaticMeshComponent) : (SplineMeshComponent != nullptr ? FGLTFNameUtilities::GetName(SplineMeshComponent) : StaticMesh->GetName());
 }
 
-void FGLTFDelayedStaticMeshTask::Process()
+void FGLTFDelayedStaticAndSplineMeshTask::Process()
 {
 	FGLTFMeshUtilities::FullyLoad(StaticMesh);
-	JsonMesh->Name = StaticMeshComponent != nullptr ? FGLTFNameUtilities::GetName(StaticMeshComponent) : StaticMesh->GetName();
 
+	const UMeshComponent* MeshComponent = StaticMeshComponent != nullptr ? StaticMeshComponent : (SplineMeshComponent != nullptr ? SplineMeshComponent : nullptr);
+
+	JsonMesh->Name = MeshComponent != nullptr ? FGLTFNameUtilities::GetName(MeshComponent) : StaticMesh->GetName();
+
+	const TArray<FStaticMaterial>& MaterialSlots = FGLTFMeshUtilities::GetMaterials(StaticMesh);
+	
+	const FGLTFMeshData* MeshData = Builder.ExportOptions->BakeMaterialInputs == EGLTFMaterialBakeMode::UseMeshData ?
+		Builder.AddUniqueMeshData(StaticMesh, StaticMeshComponent, LODIndex) : nullptr;
+
+#if WITH_EDITOR
+	if (MeshData != nullptr)
+	{
+		if (MeshData->Description.IsEmpty())
+		{
+			// TODO: report warning in case the mesh actually has data, which means we failed to extract a mesh description.
+			MeshData = nullptr;
+		}
+		else if (MeshData->BakeUsingTexCoord < 0)
+		{
+			// TODO: report warning (about missing texture coordinate for baking with mesh data).
+			MeshData = nullptr;
+		}
+	}
+#endif
+
+#if WITH_EDITORONLY_DATA
+	if (Builder.ExportOptions->bExportSourceModel)
+	{
+		ProcessMeshDescription(MaterialSlots, MeshData);
+	}
+	else
+#endif
+	{
+		ProcessRenderData(MaterialSlots, MeshData);
+	}
+}
+
+#if WITH_EDITORONLY_DATA
+void FGLTFDelayedStaticAndSplineMeshTask::ProcessMeshDescription(const TArray<FStaticMaterial>& MaterialSlots, const FGLTFMeshData* MeshData)
+{
+	FMeshDescription* MeshDescription = StaticMesh->GetMeshDescription(LODIndex);
+	const FPolygonGroupArray& PolygonGroups = MeshDescription->PolygonGroups();
+
+	if (PolygonGroups.Num() != JsonMesh->Primitives.Num()
+		|| PolygonGroups.Num() != MaterialSlots.Num())
+	{
+		return;
+	}
+
+	FStaticMeshConstAttributes Attributes(*MeshDescription);
+	TVertexAttributesConstRef<FVector3f> VertexPositions = Attributes.GetVertexPositions();
+	TVertexInstanceAttributesConstRef<FVector3f> VertexInstanceNormals = Attributes.GetVertexInstanceNormals();
+	TVertexInstanceAttributesConstRef<FVector3f> VertexInstanceTangents = Attributes.GetVertexInstanceTangents();
+	TVertexInstanceAttributesConstRef<float> VertexInstanceBinormalSigns = Attributes.GetVertexInstanceBinormalSigns();
+	TVertexInstanceAttributesConstRef<FVector2f> VertexInstanceUVs = Attributes.GetVertexInstanceUVs();
+	TVertexInstanceAttributesConstRef<FVector4f> VertexInstanceColors = Attributes.GetVertexInstanceColors();
+	TPolygonGroupAttributesConstRef<FName> PolygonGroupMaterialSlotNames = Attributes.GetPolygonGroupMaterialSlotNames();
+	TEdgeAttributesConstRef<bool> EdgeHardnesses = Attributes.GetEdgeHardnesses();
+
+	const int32 VertexCount = MeshDescription->Vertices().Num();
+	const int32 VertexInstanceCount = MeshDescription->VertexInstances().Num();
+
+	bool bHasVertexColors = Builder.ExportOptions->bExportVertexColors && MeshDescription->VertexInstanceAttributes().HasAttribute(MeshAttribute::VertexInstance::Color);
+	int32 UVCount = VertexInstanceUVs.GetNumChannels();
+
+	int32 PrimitiveIndex = 0;
+	int32 NumberOfPrimitives = MeshDescription->PolygonGroups().Num();
+	for (const FPolygonGroupID& PolygonGroupID : MeshDescription->PolygonGroups().GetElementIDs())
+	{
+		const TArrayView<const FTriangleID>& TriangleIDs = MeshDescription->GetPolygonGroupTriangles(PolygonGroupID);
+
+		if (TriangleIDs.Num() == 0)
+		{
+			PrimitiveIndex++;
+			//Do not export empty primitives.
+			continue;
+		}
+
+		FName MaterialSlotName = PolygonGroupMaterialSlotNames[PolygonGroupID];
+		FGLTFJsonPrimitive& JsonPrimitive = JsonMesh->Primitives[PrimitiveIndex];
+
+		int32 MaterialIndex = INDEX_NONE;
+		{
+			for (size_t MatIndex = 0; MatIndex < MaterialSlots.Num(); MatIndex++)
+			{
+				if (MaterialSlots[MatIndex].ImportedMaterialSlotName == MaterialSlotName)
+				{
+					MaterialIndex = MatIndex;
+					break;
+				}
+			}
+
+			if (MaterialIndex == INDEX_NONE)
+			{
+				MaterialIndex = PolygonGroupID.GetValue();
+			}
+			if (!MaterialSlots.IsValidIndex(MaterialIndex))
+			{
+				MaterialIndex = 0;
+			}
+		}
+
+		TArray<int32> OriginalIndices;
+		TArray<FVector3f> OriginalPositions;
+		TArray<FColor> OriginalVertexColors;
+		TArray<FVector3f> OriginalNormals;
+		TArray<FVector4f> OriginalTangents;
+		TArray<TArray<FVector2f>> OriginalUVs;
+
+		OriginalIndices.Reserve(VertexInstanceCount);
+		OriginalPositions.SetNumZeroed(VertexInstanceCount);
+		OriginalVertexColors.SetNumZeroed(VertexInstanceCount);
+		OriginalNormals.SetNumZeroed(VertexInstanceCount);
+		OriginalTangents.SetNumZeroed(VertexInstanceCount);
+		for (size_t UVIndex = 0; UVIndex < UVCount; UVIndex++)
+		{
+			OriginalUVs.Add(TArray<FVector2f>());
+			OriginalUVs[UVIndex].SetNumZeroed(VertexInstanceCount);
+		}
+
+		for (const FTriangleID& TriangleID : TriangleIDs)
+		{
+			TArrayView<const FVertexInstanceID> TriangleVertexInstanceIDs = MeshDescription->GetTriangleVertexInstances(TriangleID);
+			for (const FVertexInstanceID& VertexInstanceID : TriangleVertexInstanceIDs)
+			{
+				OriginalIndices.Add(VertexInstanceID);
+
+				FVector3f Position = VertexPositions[MeshDescription->GetVertexInstanceVertex(VertexInstanceID).GetValue()];
+				OriginalPositions[VertexInstanceID] = Position;
+
+				OriginalNormals[VertexInstanceID] = VertexInstanceNormals[VertexInstanceID];
+				OriginalTangents[VertexInstanceID] = FVector4f(VertexInstanceTangents[VertexInstanceID], VertexInstanceBinormalSigns[VertexInstanceID]);
+
+				for (size_t UVIndex = 0; UVIndex < UVCount; UVIndex++)
+				{
+					OriginalUVs[UVIndex][VertexInstanceID] = VertexInstanceUVs.Get(VertexInstanceID, UVIndex);
+				}
+			}
+		}
+
+		if (bHasVertexColors)
+		{
+			for (const FTriangleID& TriangleID : TriangleIDs)
+			{
+				TArrayView<const FVertexInstanceID> TriangleVertexInstanceIDs = MeshDescription->GetTriangleVertexInstances(TriangleID);
+				for (const FVertexInstanceID& VertexInstanceID : TriangleVertexInstanceIDs)
+				{
+					const FVector4f& SourceVertexColor = VertexInstanceColors[VertexInstanceID];
+					OriginalVertexColors[VertexInstanceID] = FLinearColor(SourceVertexColor).ToFColor(true);
+				}
+			}
+		}
+
+		FGLTFIndexArray Indices;
+		FGLTFPositionArray PositionBuffer;
+		FGLTFColorArray VertexColorBuffer;
+		FGLTFNormalArray Normals;
+		FGLTFTangentArray Tangents;
+		TArray<FGLTFUVArray> UVs;
+
+		//Remap Containers to contain only used data sets (per primitive)
+		// + Fill glTF Containers.
+		{
+			TArray<int32> SortedIndices = OriginalIndices;
+			TMap<int32, int32> IndexRemapper;
+			{
+				SortedIndices.Sort();
+
+				TSet<int32> SortedUniqueIndices;
+				SortedUniqueIndices.Append(SortedIndices);
+				SortedIndices.Empty();
+
+				SortedIndices = SortedUniqueIndices.Array();
+
+				for (size_t NewIndex = 0; NewIndex < SortedIndices.Num(); NewIndex++)
+				{
+					IndexRemapper.Add(SortedIndices[NewIndex], NewIndex);
+				}
+			}
+
+			int32 PrimtiveVertexCount = IndexRemapper.Num();
+
+			Indices.Reserve(PrimtiveVertexCount);
+			PositionBuffer.SetNumZeroed(PrimtiveVertexCount);
+			VertexColorBuffer.SetNumZeroed(PrimtiveVertexCount);
+			Normals.SetNumZeroed(PrimtiveVertexCount);
+			Tangents.SetNumZeroed(PrimtiveVertexCount);
+
+			UVs.AddDefaulted(UVCount);
+			for (size_t UVIndex = 0; UVIndex < UVCount; UVIndex++)
+			{
+				UVs[UVIndex].SetNumZeroed(PrimtiveVertexCount);
+			}
+
+			for (const int32& OriginalIndex : OriginalIndices)
+			{
+				int32 NewIndex = IndexRemapper[OriginalIndex];
+				Indices.Add(NewIndex);
+
+				FVector3f Position = OriginalPositions[OriginalIndex];
+				if (SplineMeshComponent)
+				{
+					//SplineMeshComponent provided
+					//Fix the Positions for Splines:
+					const FTransform3f SliceTransform = FTransform3f(SplineMeshComponent->CalcSliceTransform(USplineMeshComponent::GetAxisValueRef(Position, SplineMeshComponent->ForwardAxis)));
+					USplineMeshComponent::GetAxisValueRef(Position, SplineMeshComponent->ForwardAxis) = 0;
+					Position = SliceTransform.TransformPosition(Position);
+				}
+
+				PositionBuffer[NewIndex] = Position;
+
+				if (bHasVertexColors)
+				{
+					VertexColorBuffer[NewIndex] = OriginalVertexColors[OriginalIndex];
+				}
+
+				Normals[NewIndex] = OriginalNormals[OriginalIndex];
+				Tangents[NewIndex] = OriginalTangents[OriginalIndex];
+
+				for (size_t UVIndex = 0; UVIndex < UVCount; UVIndex++)
+				{
+					UVs[UVIndex][NewIndex] = OriginalUVs[UVIndex][OriginalIndex];
+				}
+			}
+		}
+
+		if (Tangents.Num() > 0 && DoesBufferHasZeroVector(Tangents))
+		{
+			//Do not Export Tangents list that is zeroed out.
+			Tangents.Reset();
+		}
+
+		if (Normals.Num() > 0 && DoesBufferHasZeroVector(Normals))
+		{
+			//Do not Export Normals list that is zeroed out.
+			Normals.Reset();
+		}
+
+		//Set glTF Primitive:
+		JsonPrimitive.Indices = Builder.AddUniqueIndexAccessor(Indices, StaticMesh->GetName() + (NumberOfPrimitives > 1 ? (TEXT("_") + FString::FromInt(PrimitiveIndex)) : TEXT("")));
+		JsonPrimitive.Attributes.Position = Builder.AddUniquePositionAccessor(PositionBuffer);
+		if (bHasVertexColors) JsonPrimitive.Attributes.Color0 = Builder.AddUniqueColorAccessor(VertexColorBuffer);
+		JsonPrimitive.Attributes.Normal = Builder.AddUniqueNormalAccessor(Normals);
+		JsonPrimitive.Attributes.Tangent = Builder.AddUniqueTangentAccessor(Tangents);
+		JsonPrimitive.Attributes.TexCoords.AddUninitialized(UVCount);
+		for (size_t UVIndex = 0; UVIndex < UVCount; UVIndex++)
+		{
+			JsonPrimitive.Attributes.TexCoords[UVIndex] = Builder.AddUniqueUVAccessor(UVs[UVIndex]);
+		}
+		//
+		const UMaterialInterface* Material = Materials.IsValidIndex(MaterialIndex) ? Materials[MaterialIndex] : MaterialSlots[MaterialIndex].MaterialInterface;
+		JsonPrimitive.Material = Builder.AddUniqueMaterial(Material, MeshData, { MaterialIndex });
+
+		//Validations:
+		if (JsonPrimitive.Attributes.Position == nullptr)
+		{
+			Builder.LogError(
+				FString::Printf(TEXT("Failed to export vertex positions related to material slot %d (%s) in static mesh %s"),
+					0,
+					*JsonPrimitive.Material->Name,
+					*JsonMesh->Name
+				));
+		}
+
+		PrimitiveIndex++;
+	}
+}
+#endif
+
+void FGLTFDelayedStaticAndSplineMeshTask::ProcessRenderData(const TArray<FStaticMaterial>& MaterialSlots, const FGLTFMeshData* MeshData)
+{
 	const FStaticMeshLODResources& RenderData = FGLTFMeshUtilities::GetRenderData(StaticMesh, LODIndex);
+
 	const FPositionVertexBuffer& PositionBuffer = RenderData.VertexBuffers.PositionVertexBuffer;
 	const FStaticMeshVertexBuffer& VertexBuffer = RenderData.VertexBuffers.StaticMeshVertexBuffer;
 	const FColorVertexBuffer* ColorBuffer = &RenderData.VertexBuffers.ColorVertexBuffer; // TODO: add support for overriding color buffer by component
@@ -150,32 +436,15 @@ void FGLTFDelayedStaticMeshTask::Process()
 		const FStaticMeshComponentLODInfo& LODInfo = StaticMeshComponent->LODData[LODIndex];
 		ColorBuffer = LODInfo.OverrideVertexColors != nullptr ? LODInfo.OverrideVertexColors : ColorBuffer;
 	}
-
-	const FGLTFMeshData* MeshData = Builder.ExportOptions->BakeMaterialInputs == EGLTFMaterialBakeMode::UseMeshData ?
-		Builder.AddUniqueMeshData(StaticMesh, StaticMeshComponent, LODIndex) : nullptr;
-
-#if WITH_EDITOR
-	if (MeshData != nullptr)
+	else if (SplineMeshComponent != nullptr && SplineMeshComponent->LODData.IsValidIndex(LODIndex))
 	{
-		if (MeshData->Description.IsEmpty())
-		{
-			// TODO: report warning in case the mesh actually has data, which means we failed to extract a mesh description.
-			MeshData = nullptr;
-		}
-		else if (MeshData->BakeUsingTexCoord < 0)
-		{
-			// TODO: report warning (about missing texture coordinate for baking with mesh data).
-			MeshData = nullptr;
-		}
+		const FStaticMeshComponentLODInfo& LODInfo = SplineMeshComponent->LODData[LODIndex];
+		ColorBuffer = LODInfo.OverrideVertexColors != nullptr ? LODInfo.OverrideVertexColors : ColorBuffer;
 	}
-#endif
 
 	ValidateVertexBuffer(Builder, &VertexBuffer, *StaticMesh->GetName());
 
-	const TArray<FStaticMaterial>& MaterialSlots = FGLTFMeshUtilities::GetMaterials(StaticMesh);
-	const int32 MaterialCount = MaterialSlots.Num();
-
-	for (int32 MaterialIndex = 0; MaterialIndex < MaterialCount; ++MaterialIndex)
+	for (int32 MaterialIndex = 0; MaterialIndex < MaterialSlots.Num(); ++MaterialIndex)
 	{
 		const FGLTFIndexArray SectionIndices = FGLTFMeshUtilities::GetSectionIndices(RenderData, MaterialIndex);
 		const FGLTFMeshSection* ConvertedSection = MeshSectionConverter.GetOrAdd(StaticMesh, LODIndex, SectionIndices);
@@ -183,7 +452,42 @@ void FGLTFDelayedStaticMeshTask::Process()
 		FGLTFJsonPrimitive& JsonPrimitive = JsonMesh->Primitives[MaterialIndex];
 		JsonPrimitive.Indices = Builder.AddUniqueIndexAccessor(ConvertedSection);
 
-		JsonPrimitive.Attributes.Position = Builder.AddUniquePositionAccessor(ConvertedSection, &PositionBuffer);
+		if (!JsonPrimitive.Indices || JsonPrimitive.Indices->Count == 0)
+		{
+			//Do not export empty primitives.
+			continue;
+		}
+
+		if (SplineMeshComponent)
+		{
+			//fix for Splines:
+			FPositionVertexBuffer* TransformedPositionBuffer = new FPositionVertexBuffer();
+
+			TransformedPositionBuffer->Init(PositionBuffer.GetNumVertices(), true);
+
+			const uint32 VertexCount = PositionBuffer.GetNumVertices();
+			const uint32 Stride = PositionBuffer.GetStride();
+
+			const TUniquePtr<IGLTFBufferAdapter> SourceBuffer = IGLTFBufferAdapter::GetPositions(&PositionBuffer);
+			const uint8* SourceData = SourceBuffer->GetData();
+
+			for (uint32 VertexIndex = 0; VertexIndex < VertexCount; ++VertexIndex, SourceData+=Stride)
+			{
+				FVector3f& VertexPosition = TransformedPositionBuffer->VertexPosition(VertexIndex);
+				VertexPosition = *reinterpret_cast<const FVector3f*>(SourceData);
+
+				const FTransform3f SliceTransform = FTransform3f(SplineMeshComponent->CalcSliceTransform(USplineMeshComponent::GetAxisValueRef(VertexPosition, SplineMeshComponent->ForwardAxis)));
+				USplineMeshComponent::GetAxisValueRef(VertexPosition, SplineMeshComponent->ForwardAxis) = 0;
+				VertexPosition = SliceTransform.TransformPosition(VertexPosition);
+			}
+
+			JsonPrimitive.Attributes.Position = Builder.AddUniquePositionAccessor(ConvertedSection, TransformedPositionBuffer);
+		}
+		else
+		{
+			JsonPrimitive.Attributes.Position = Builder.AddUniquePositionAccessor(ConvertedSection, &PositionBuffer);
+		}
+
 		if (JsonPrimitive.Attributes.Position == nullptr)
 		{
 			Builder.LogError(
@@ -217,6 +521,7 @@ void FGLTFDelayedStaticMeshTask::Process()
 	}
 }
 
+
 FString FGLTFDelayedSkeletalMeshTask::GetName()
 {
 	return SkeletalMeshComponent != nullptr ? FGLTFNameUtilities::GetName(SkeletalMeshComponent) : SkeletalMesh->GetName();
@@ -227,32 +532,6 @@ void FGLTFDelayedSkeletalMeshTask::Process()
 	FGLTFMeshUtilities::FullyLoad(SkeletalMesh);
 	JsonMesh->Name = SkeletalMeshComponent != nullptr ? FGLTFNameUtilities::GetName(SkeletalMeshComponent) : SkeletalMesh->GetName();
 
-	const FSkeletalMeshLODRenderData& RenderData = FGLTFMeshUtilities::GetRenderData(SkeletalMesh, LODIndex);
-	const FPositionVertexBuffer& PositionBuffer = RenderData.StaticVertexBuffers.PositionVertexBuffer;
-	const FStaticMeshVertexBuffer& VertexBuffer = RenderData.StaticVertexBuffers.StaticMeshVertexBuffer;
-	const FColorVertexBuffer* ColorBuffer = &RenderData.StaticVertexBuffers.ColorVertexBuffer; // TODO: add support for overriding color buffer by component
-	const FSkinWeightVertexBuffer* SkinWeightBuffer = RenderData.GetSkinWeightVertexBuffer(); // TODO: add support for overriding skin weight buffer by component
-	// TODO: add support for skin weight profiles?
-	// TODO: add support for morph targets
-
-	if (Builder.ExportOptions->bExportVertexColors && HasVertexColors(ColorBuffer))
-	{
-		Builder.LogSuggestion(FString::Printf(
-			TEXT("Vertex colors in mesh %s will act as a multiplier for base color in glTF, regardless of material, which may produce undesirable results"),
-			*SkeletalMesh->GetName()));
-	}
-	else
-	{
-		ColorBuffer = nullptr;
-	}
-
-	if (SkeletalMeshComponent != nullptr && SkeletalMeshComponent->LODInfo.IsValidIndex(LODIndex))
-	{
-		const FSkelMeshComponentLODInfo& LODInfo = SkeletalMeshComponent->LODInfo[LODIndex];
-		ColorBuffer = LODInfo.OverrideVertexColors != nullptr ? LODInfo.OverrideVertexColors : ColorBuffer;
-		SkinWeightBuffer = LODInfo.OverrideSkinWeights != nullptr ? LODInfo.OverrideSkinWeights : SkinWeightBuffer;
-	}
-
 	const FGLTFMeshData* MeshData = Builder.ExportOptions->BakeMaterialInputs == EGLTFMaterialBakeMode::UseMeshData ?
 		Builder.AddUniqueMeshData(SkeletalMesh, SkeletalMeshComponent, LODIndex) : nullptr;
 
@@ -272,160 +551,283 @@ void FGLTFDelayedSkeletalMeshTask::Process()
 	}
 #endif
 
-	ValidateVertexBuffer(Builder, &VertexBuffer, *SkeletalMesh->GetName());
-
 	const TArray<FSkeletalMaterial>& MaterialSlots = FGLTFMeshUtilities::GetMaterials(SkeletalMesh);
-	const int32 MaterialCount = MaterialSlots.Num();
 
-	for (int32 MaterialIndex = 0; MaterialIndex < MaterialCount; ++MaterialIndex)
+#if WITH_EDITORONLY_DATA
+	if (Builder.ExportOptions->bExportSourceModel)
 	{
-		const FGLTFIndexArray SectionIndices = FGLTFMeshUtilities::GetSectionIndices(RenderData, MaterialIndex);
-		const FGLTFMeshSection* ConvertedSection = MeshSectionConverter.GetOrAdd(SkeletalMesh, LODIndex, SectionIndices);
+		ProcessSourceModel(MaterialSlots, MeshData);
+	}
+	else
+#endif
+	{
+		ProcessRenderData(MaterialSlots, MeshData);
+	}
+}
 
-		FGLTFJsonPrimitive& JsonPrimitive = JsonMesh->Primitives[MaterialIndex];
-		JsonPrimitive.Indices = Builder.AddUniqueIndexAccessor(ConvertedSection);
+#if WITH_EDITORONLY_DATA
+void FGLTFDelayedSkeletalMeshTask::ProcessSourceModel(const TArray<FSkeletalMaterial>& MaterialSlots, const FGLTFMeshData* MeshData)
+{
+	FSkeletalMeshModel* ImportedSkeletalMeshModel = SkeletalMesh->GetImportedModel();
+	if (!ImportedSkeletalMeshModel->LODModels.IsValidIndex(LODIndex))
+	{
+		//TODO: Log Error
+		return;
+	}
 
-		JsonPrimitive.Attributes.Position = Builder.AddUniquePositionAccessor(ConvertedSection, &PositionBuffer);
-		if (JsonPrimitive.Attributes.Position == nullptr)
+	bool bExportVertexColors = Builder.ExportOptions->bExportVertexColors;
+
+	FSkeletalMeshLODModel& SourceModel = ImportedSkeletalMeshModel->LODModels[LODIndex];
+	const int32 SectionCount = SourceModel.Sections.Num();
+	const int32 NumTexCoords = SourceModel.NumTexCoords;
+
+	if (JsonMesh->Primitives.Num() != SectionCount)
+	{
+		//TODO: Log Error
+		return;
+	}
+
+	TArray<FSoftSkinVertex> SoftSkinVertices;
+	SourceModel.GetNonClothVertices(SoftSkinVertices);
+
+	int32 ClothSectionVertexRemoveOffset = 0;
+	for (int32 SectionIndex = 0; SectionIndex < SectionCount; ++SectionIndex)
+	{
+		FGLTFJsonPrimitive& JsonPrimitive = JsonMesh->Primitives[SectionIndex];
+
+		const FSkelMeshSection& Section = SourceModel.Sections[SectionIndex];
+		//Section.MaterialIndex
+		if (Section.HasClothingData() || !(MaterialSlots.IsValidIndex(Section.MaterialIndex)))
 		{
-			Builder.LogError(
-				FString::Printf(TEXT("Failed to export vertex positions related to material slot %d (%s) in skeletal mesh %s"),
-					MaterialIndex,
-					*MaterialSlots[MaterialIndex].MaterialSlotName.ToString(),
-					*ConvertedSection->ToString()
-				));
+			ClothSectionVertexRemoveOffset += Section.GetNumVertices();
+			continue;
 		}
 
-		if (ColorBuffer != nullptr)
+		int32 TriangleCount = Section.NumTriangles;
+
+		TArray<int32> OriginalIndices;
+		OriginalIndices.Reserve(TriangleCount * 3);
+
+		// Copy over the index buffer into the FBX polygons set.
+		for (int32 TriangleIndex = 0; TriangleIndex < TriangleCount; ++TriangleIndex)
 		{
-			JsonPrimitive.Attributes.Color0 = Builder.AddUniqueColorAccessor(ConvertedSection, ColorBuffer);
+			for (int32 PointIndex = 0; PointIndex < 3; PointIndex++)
+			{
+				int32 VertexPositionIndex = SourceModel.IndexBuffer[Section.BaseIndex + ((TriangleIndex * 3) + PointIndex)] - ClothSectionVertexRemoveOffset;
+				OriginalIndices.Add(VertexPositionIndex);
+			}
 		}
 
-		// TODO: report warning if both Mesh Quantization (export options) and Use High Precision Tangent Basis (vertex buffer) are disabled
-		JsonPrimitive.Attributes.Normal = Builder.AddUniqueNormalAccessor(ConvertedSection, &VertexBuffer);
-		JsonPrimitive.Attributes.Tangent = Builder.AddUniqueTangentAccessor(ConvertedSection, &VertexBuffer);
+		struct UIndexProcessedHelper
+		{
+			int32 NewIndex;
+			bool Processed;
+			UIndexProcessedHelper(const int32& Index)
+				: NewIndex(Index)
+				, Processed(false)
+			{
 
-		const uint32 UVCount = VertexBuffer.GetNumTexCoords();
-		// TODO: report warning or option to limit UV channels since most viewers don't support more than 2?
-		JsonPrimitive.Attributes.TexCoords.AddUninitialized(UVCount);
+			}
+		};
 
-		for (uint32 UVIndex = 0; UVIndex < UVCount; ++UVIndex)
+		TMap<int32, UIndexProcessedHelper> IndexRemapper;
 		{
-			JsonPrimitive.Attributes.TexCoords[UVIndex] = Builder.AddUniqueUVAccessor(ConvertedSection, &VertexBuffer, UVIndex);
+			TArray<int32> SortedIndices = OriginalIndices;
+			SortedIndices.Sort();
+
+			TSet<int32> SortedUniqueIndices;
+			SortedUniqueIndices.Append(SortedIndices);
+			SortedIndices.Empty();
+
+			SortedIndices = SortedUniqueIndices.Array();
+
+			for (size_t NewIndex = 0; NewIndex < SortedIndices.Num(); NewIndex++)
+			{
+				IndexRemapper.Add(SortedIndices[NewIndex], NewIndex);
+			}
 		}
 
+		FGLTFIndexArray Indices;
+		FGLTFPositionArray PositionBuffer;
+		FGLTFColorArray VertexColorBuffer;
+		FGLTFNormalArray Normals;
+		FGLTFTangentArray Tangents;
+		TArray<FGLTFUVArray> UVs;
+		TArray<FGLTFJointInfluenceArray> JointInfluences; //Per Group
+		TArray<FGLTFJointWeightArray> JointWeights; //Per Group
+
+		int32 PrimtiveVertexCount = IndexRemapper.Num();
+
+		Indices.Reserve(OriginalIndices.Num());
+
+		PositionBuffer.SetNumZeroed(PrimtiveVertexCount);
+		VertexColorBuffer.SetNumZeroed(PrimtiveVertexCount);
+		Normals.SetNumZeroed(PrimtiveVertexCount);
+		Tangents.SetNumZeroed(PrimtiveVertexCount);
+
+		UVs.AddDefaulted(NumTexCoords);
+		for (size_t UVIndex = 0; UVIndex < NumTexCoords; UVIndex++)
+		{
+			UVs[UVIndex].SetNumZeroed(PrimtiveVertexCount);
+		}
+
+		const uint32 GroupCount = (Section.GetMaxBoneInfluences() + 3) / 4;
 		if (Builder.ExportOptions->bExportVertexSkinWeights)
 		{
-			const uint32 GroupCount = (SkinWeightBuffer->GetMaxBoneInfluences() + 3) / 4;
 			// TODO: report warning or option to limit groups (of joints and weights) since most viewers don't support more than one?
-			JsonPrimitive.Attributes.Joints.AddUninitialized(GroupCount);
-			JsonPrimitive.Attributes.Weights.AddUninitialized(GroupCount);
+			JointInfluences.AddDefaulted(GroupCount);
+			JointWeights.AddDefaulted(GroupCount);
+			for (size_t GroupIndex = 0; GroupIndex < GroupCount; GroupIndex++)
+			{
+				JointInfluences[GroupIndex].SetNumZeroed(PrimtiveVertexCount);
+				JointWeights[GroupIndex].SetNumZeroed(PrimtiveVertexCount);
+			}
+		}
 
-			for (uint32 GroupIndex = 0; GroupIndex < GroupCount; ++GroupIndex)
+		for (const int32& OriginalIndex : OriginalIndices)
+		{
+			UIndexProcessedHelper& IndexProcessedHelper = IndexRemapper[OriginalIndex];
+			int32 NewIndex = IndexProcessedHelper.NewIndex;
+
+			Indices.Add(NewIndex);
+
+			if (IndexProcessedHelper.Processed)
 			{
-				JsonPrimitive.Attributes.Joints[GroupIndex] = Builder.AddUniqueJointAccessor(ConvertedSection, SkinWeightBuffer, GroupIndex * 4);
-				JsonPrimitive.Attributes.Weights[GroupIndex] = Builder.AddUniqueWeightAccessor(ConvertedSection, SkinWeightBuffer, GroupIndex * 4);
+				continue;
+			}
+			IndexProcessedHelper.Processed = true;
+
+			const FSoftSkinVertex& OriginalSoftSkinVertex = SoftSkinVertices[OriginalIndex];
+
+
+			PositionBuffer[NewIndex] = OriginalSoftSkinVertex.Position;
+
+			if (bExportVertexColors)
+			{
+				VertexColorBuffer[NewIndex] = OriginalSoftSkinVertex.Color;
+			}
+
+			Normals[NewIndex] = OriginalSoftSkinVertex.TangentZ;
+			Tangents[NewIndex] = OriginalSoftSkinVertex.TangentX;
+
+			for (size_t UVIndex = 0; UVIndex < NumTexCoords; UVIndex++)
+			{
+				UVs[UVIndex][NewIndex] = OriginalSoftSkinVertex.UVs[UVIndex];
+			}
+
+			if (Builder.ExportOptions->bExportVertexSkinWeights)
+			{
+				for (size_t GroupIndex = 0; GroupIndex < GroupCount; GroupIndex++)
+				{
+					for (size_t GroupBoneInfluenceCounter = 0; GroupBoneInfluenceCounter < 4; GroupBoneInfluenceCounter++)
+					{
+						int8 InfluenceIndex = GroupIndex * 4 + GroupBoneInfluenceCounter;
+
+						JointInfluences[GroupIndex][NewIndex][GroupBoneInfluenceCounter] = Section.BoneMap[OriginalSoftSkinVertex.InfluenceBones[InfluenceIndex]];
+						JointWeights[GroupIndex][NewIndex][GroupBoneInfluenceCounter] = OriginalSoftSkinVertex.InfluenceWeights[InfluenceIndex];
+					}
+				}
 			}
 		}
 
-		const UMaterialInterface* Material = Materials[MaterialIndex];
-		JsonPrimitive.Material = Builder.AddUniqueMaterial(Material, MeshData, SectionIndices);
-	}
-}
+		if (Tangents.Num() > 0 && DoesBufferHasZeroVector(Tangents))
+		{
+			//Do not Export Tangents list that is zeroed out.
+			Tangents.Reset();
+		}
 
+		if (Normals.Num() > 0 && DoesBufferHasZeroVector(Normals))
+		{
+			//Do not Export Normals list that is zeroed out.
+			Normals.Reset();
+		}
 
-FString FGLTFDelayedSplineMeshTask::GetName()
-{
-	return FGLTFNameUtilities::GetName(&SplineMeshComponent);
+		//Set glTF Primitive:
+		JsonPrimitive.Indices = Builder.AddUniqueIndexAccessor(Indices, SkeletalMesh->GetName() + (SectionCount > 1 ? (TEXT("_") + FString::FromInt(SectionIndex)) : TEXT("")));
+		JsonPrimitive.Attributes.Position = Builder.AddUniquePositionAccessor(PositionBuffer);
+		if (bExportVertexColors) JsonPrimitive.Attributes.Color0 = Builder.AddUniqueColorAccessor(VertexColorBuffer);
+		JsonPrimitive.Attributes.Normal = Builder.AddUniqueNormalAccessor(Normals);
+		JsonPrimitive.Attributes.Tangent = Builder.AddUniqueTangentAccessor(Tangents);
+		JsonPrimitive.Attributes.TexCoords.AddUninitialized(NumTexCoords);
+		for (size_t UVIndex = 0; UVIndex < NumTexCoords; UVIndex++)
+		{
+			JsonPrimitive.Attributes.TexCoords[UVIndex] = Builder.AddUniqueUVAccessor(UVs[UVIndex]);
+		}
+
+		if (Builder.ExportOptions->bExportVertexSkinWeights)
+		{
+			JsonPrimitive.Attributes.Joints.AddUninitialized(GroupCount);
+			JsonPrimitive.Attributes.Weights.AddUninitialized(GroupCount);
+			for (size_t GroupCountIndex = 0; GroupCountIndex < GroupCount; GroupCountIndex++)
+			{
+				JsonPrimitive.Attributes.Joints[GroupCountIndex] = Builder.AddUniqueJointAccessor(JointInfluences[GroupCountIndex]);
+				JsonPrimitive.Attributes.Weights[GroupCountIndex] = Builder.AddUniqueWeightAccessor(JointWeights[GroupCountIndex]);
+			}
+		}
+
+		//
+		const UMaterialInterface* Material = Materials.IsValidIndex(Section.MaterialIndex) ? Materials[Section.MaterialIndex] : MaterialSlots[Section.MaterialIndex].MaterialInterface;
+		JsonPrimitive.Material = Builder.AddUniqueMaterial(Material, MeshData, { SectionIndex });
+
+		//Validations:
+		if (JsonPrimitive.Attributes.Position == nullptr)
+		{
+			Builder.LogError(
+				FString::Printf(TEXT("Failed to export vertex positions related to material slot %d (%s) in static mesh %s"),
+					0,
+					*JsonPrimitive.Material->Name,
+					*JsonMesh->Name
+				));
+		}
+	}
 }
+#endif
 
-void FGLTFDelayedSplineMeshTask::Process()
+void FGLTFDelayedSkeletalMeshTask::ProcessRenderData(const TArray<FSkeletalMaterial>& MaterialSlots, const FGLTFMeshData* MeshData)
 {
-	FGLTFMeshUtilities::FullyLoad(&StaticMesh);
-	JsonMesh->Name = GetName();
-
-	const FStaticMeshLODResources& RenderData = FGLTFMeshUtilities::GetRenderData(&StaticMesh, LODIndex);
-	const FPositionVertexBuffer& PositionBuffer = RenderData.VertexBuffers.PositionVertexBuffer;
-	const FStaticMeshVertexBuffer& VertexBuffer = RenderData.VertexBuffers.StaticMeshVertexBuffer;
-	const FColorVertexBuffer* ColorBuffer = &RenderData.VertexBuffers.ColorVertexBuffer; // TODO: add support for overriding color buffer by component
+	const FSkeletalMeshLODRenderData& RenderData = FGLTFMeshUtilities::GetRenderData(SkeletalMesh, LODIndex);
+	const FPositionVertexBuffer& PositionBuffer = RenderData.StaticVertexBuffers.PositionVertexBuffer;
+	const FStaticMeshVertexBuffer& VertexBuffer = RenderData.StaticVertexBuffers.StaticMeshVertexBuffer;
+	const FColorVertexBuffer* ColorBuffer = &RenderData.StaticVertexBuffers.ColorVertexBuffer; // TODO: add support for overriding color buffer by component
+	const FSkinWeightVertexBuffer* SkinWeightBuffer = RenderData.GetSkinWeightVertexBuffer(); // TODO: add support for overriding skin weight buffer by component
+	// TODO: add support for skin weight profiles?
+	// TODO: add support for morph targets
 
 	if (Builder.ExportOptions->bExportVertexColors && HasVertexColors(ColorBuffer))
 	{
 		Builder.LogSuggestion(FString::Printf(
 			TEXT("Vertex colors in mesh %s will act as a multiplier for base color in glTF, regardless of material, which may produce undesirable results"),
-			*StaticMesh.GetName()));
+			*SkeletalMesh->GetName()));
 	}
 	else
 	{
 		ColorBuffer = nullptr;
 	}
 
-	if (SplineMeshComponent.LODData.IsValidIndex(LODIndex))
+	if (SkeletalMeshComponent != nullptr && SkeletalMeshComponent->LODInfo.IsValidIndex(LODIndex))
 	{
-		const FStaticMeshComponentLODInfo& LODInfo = SplineMeshComponent.LODData[LODIndex];
+		const FSkelMeshComponentLODInfo& LODInfo = SkeletalMeshComponent->LODInfo[LODIndex];
 		ColorBuffer = LODInfo.OverrideVertexColors != nullptr ? LODInfo.OverrideVertexColors : ColorBuffer;
+		SkinWeightBuffer = LODInfo.OverrideSkinWeights != nullptr ? LODInfo.OverrideSkinWeights : SkinWeightBuffer;
 	}
 
-	const FGLTFMeshData* MeshData = Builder.ExportOptions->BakeMaterialInputs == EGLTFMaterialBakeMode::UseMeshData ?
-	Builder.AddUniqueMeshData(&StaticMesh, &SplineMeshComponent, LODIndex) : nullptr;
-
-	#if WITH_EDITOR
-	if (MeshData != nullptr)
-	{
-		if (MeshData->Description.IsEmpty())
-		{
-			// TODO: report warning in case the mesh actually has data, which means we failed to extract a mesh description.
-			MeshData = nullptr;
-		}
-		else if (MeshData->BakeUsingTexCoord < 0)
-		{
-			// TODO: report warning (about missing texture coordinate for baking with mesh data).
-			MeshData = nullptr;
-		}
-	}
-	#endif
-
-	ValidateVertexBuffer(Builder, &VertexBuffer, *StaticMesh.GetName());
+	ValidateVertexBuffer(Builder, &VertexBuffer, *SkeletalMesh->GetName());
 
-	const TArray<FStaticMaterial>& MaterialSlots = FGLTFMeshUtilities::GetMaterials(&StaticMesh);
 	const int32 MaterialCount = MaterialSlots.Num();
 
 	for (int32 MaterialIndex = 0; MaterialIndex < MaterialCount; ++MaterialIndex)
 	{
 		const FGLTFIndexArray SectionIndices = FGLTFMeshUtilities::GetSectionIndices(RenderData, MaterialIndex);
-		const FGLTFMeshSection* ConvertedSection = MeshSectionConverter.GetOrAdd(&StaticMesh, LODIndex, SectionIndices);
+		const FGLTFMeshSection* ConvertedSection = MeshSectionConverter.GetOrAdd(SkeletalMesh, LODIndex, SectionIndices);
 
 		FGLTFJsonPrimitive& JsonPrimitive = JsonMesh->Primitives[MaterialIndex];
 		JsonPrimitive.Indices = Builder.AddUniqueIndexAccessor(ConvertedSection);
 
-		FPositionVertexBuffer* TransformedPositionBuffer = new FPositionVertexBuffer();
-		
-		{//fix for Splines:
-			TransformedPositionBuffer->Init(PositionBuffer.GetNumVertices(), true);
-
-			const uint32 VertexCount = PositionBuffer.GetNumVertices();
-			const uint32 Stride = PositionBuffer.GetStride();
-
-			const TUniquePtr<IGLTFBufferAdapter> SourceBuffer = IGLTFBufferAdapter::GetPositions(&PositionBuffer);
-			const uint8* SourceData = SourceBuffer->GetData();
-
-			for (uint32 VertexIndex = 0; VertexIndex < VertexCount; ++VertexIndex)
-			{
-				FVector3f Position = *reinterpret_cast<const FVector3f*>(SourceData + Stride * VertexIndex);
-
-				const FTransform3f SliceTransform = FTransform3f(SplineMeshComponent.CalcSliceTransform(USplineMeshComponent::GetAxisValueRef(Position, SplineMeshComponent.ForwardAxis)));
-				USplineMeshComponent::GetAxisValueRef(Position, SplineMeshComponent.ForwardAxis) = 0;
-				Position = SliceTransform.TransformPosition(Position);
-
-				TransformedPositionBuffer->VertexPosition(VertexIndex) = FVector3f(Position.X, Position.Y, Position.Z);
-			}
-		}
-
-		JsonPrimitive.Attributes.Position = Builder.AddUniquePositionAccessor(ConvertedSection, TransformedPositionBuffer);
+		JsonPrimitive.Attributes.Position = Builder.AddUniquePositionAccessor(ConvertedSection, &PositionBuffer);
 		if (JsonPrimitive.Attributes.Position == nullptr)
 		{
 			Builder.LogError(
-				FString::Printf(TEXT("Failed to export vertex positions related to material slot %d (%s) in static mesh %s"),
+				FString::Printf(TEXT("Failed to export vertex positions related to material slot %d (%s) in skeletal mesh %s"),
 					MaterialIndex,
 					*MaterialSlots[MaterialIndex].MaterialSlotName.ToString(),
 					*ConvertedSection->ToString()
@@ -450,6 +852,20 @@ void FGLTFDelayedSplineMeshTask::Process()
 			JsonPrimitive.Attributes.TexCoords[UVIndex] = Builder.AddUniqueUVAccessor(ConvertedSection, &VertexBuffer, UVIndex);
 		}
 
+		if (Builder.ExportOptions->bExportVertexSkinWeights)
+		{
+			const uint32 GroupCount = (SkinWeightBuffer->GetMaxBoneInfluences() + 3) / 4;
+			// TODO: report warning or option to limit groups (of joints and weights) since most viewers don't support more than one?
+			JsonPrimitive.Attributes.Joints.AddUninitialized(GroupCount);
+			JsonPrimitive.Attributes.Weights.AddUninitialized(GroupCount);
+
+			for (uint32 GroupIndex = 0; GroupIndex < GroupCount; ++GroupIndex)
+			{
+				JsonPrimitive.Attributes.Joints[GroupIndex] = Builder.AddUniqueJointAccessor(ConvertedSection, SkinWeightBuffer, GroupIndex * 4);
+				JsonPrimitive.Attributes.Weights[GroupIndex] = Builder.AddUniqueWeightAccessor(ConvertedSection, SkinWeightBuffer, GroupIndex * 4);
+			}
+		}
+
 		const UMaterialInterface* Material = Materials[MaterialIndex];
 		JsonPrimitive.Material = Builder.AddUniqueMaterial(Material, MeshData, SectionIndices);
 	}
@@ -488,13 +904,23 @@ void FGLTFDelayedLandscapeTask::Process()
 	const int32 VertexCount = FMath::Square(ComponentSizeQuads + 1);
 	const int32 TriangleCount = FMath::Square(ComponentSizeQuads) * 2;
 
-	FPositionVertexBuffer* PositionBuffer = new FPositionVertexBuffer();
-	PositionBuffer->Init(VertexCount, true);
+	FGLTFIndexArray Indices;
+	FGLTFPositionArray PositionBuffer;
+	FGLTFColorArray VertexColorBuffer;
+	FGLTFNormalArray Normals;
+	FGLTFTangentArray Tangents;
+	FGLTFUVArray UV;
+
+	Indices.Reserve(FMath::Square(ComponentSizeQuads) * 2 * 3);
+	PositionBuffer.SetNumZeroed(VertexCount);
+	VertexColorBuffer.SetNumZeroed(VertexCount);
+	Normals.SetNumZeroed(VertexCount);
+	Tangents.SetNumZeroed(VertexCount);
+	UV.SetNumZeroed(VertexCount);
 
 	FGLTFJsonPrimitive& JsonPrimitive = JsonMesh->Primitives[0];
 	TArray<uint8> VisibilityData;
-	VisibilityData.Empty(VertexCount);
-	VisibilityData.AddZeroed(VertexCount);
+	VisibilityData.SetNumZeroed(VertexCount);
 
 	int OffsetX = Landscape->LandscapeSectionOffset.X;
 	int OffsetY = Landscape->LandscapeSectionOffset.Y;
@@ -526,11 +952,6 @@ void FGLTFDelayedLandscapeTask::Process()
 		}
 	}
 
-	FGLTFNormalArray* Normals = new FGLTFNormalArray();
-	Normals->SetNumZeroed(VertexCount);
-	FGLTFUVArray* UVs = new FGLTFUVArray(); //UVIndex==0
-	UVs->SetNumZeroed(VertexCount);
-
 	for (int32 VertexIndex = 0; VertexIndex < VertexCount; VertexIndex++)
 	{
 		int32 VertX, VertY;
@@ -538,32 +959,29 @@ void FGLTFDelayedLandscapeTask::Process()
 
 		FVector3f Position;
 		FVector3f Normal;
-		FVector2f UV;
-		CDI.GetPositionNormalUV(VertX, VertY, Position, Normal, UV);
+		FVector2f UVElement;
+		CDI.GetPositionNormalUV(VertX, VertY, Position, Normal, UVElement);
 
-		PositionBuffer->VertexPosition(VertexIndex) = Position;
-		(*Normals)[VertexIndex] = Normal;
-		(*UVs)[VertexIndex] = UV;
+		PositionBuffer[VertexIndex] = Position;
+		Normals[VertexIndex] = Normal;
+		UV[VertexIndex] = UVElement;
 	}
 
 	const int32 VisThreshold = 170;
 	
-	FGLTFIndexArray* LandscapeIndices = new FGLTFIndexArray();
-	LandscapeIndices->Reserve(FMath::Square(ComponentSizeQuads) * 2 * 3);
-
 	for (int32 Y = 0; Y < ComponentSizeQuads; Y++)
 	{
 		for (int32 X = 0; X < ComponentSizeQuads; X++)
 		{
 			if (VisibilityData[Y * (ComponentSizeQuads + 1) + X] < VisThreshold)
 			{
-				LandscapeIndices->Push((X + 0) + (Y + 0) * (ComponentSizeQuads + 1));
-				LandscapeIndices->Push((X + 1) + (Y + 1) * (ComponentSizeQuads + 1));
-				LandscapeIndices->Push((X + 1) + (Y + 0) * (ComponentSizeQuads + 1));
+				Indices.Push((X + 0) + (Y + 0) * (ComponentSizeQuads + 1));
+				Indices.Push((X + 1) + (Y + 1) * (ComponentSizeQuads + 1));
+				Indices.Push((X + 1) + (Y + 0) * (ComponentSizeQuads + 1));
 
-				LandscapeIndices->Push((X + 0) + (Y + 0) * (ComponentSizeQuads + 1));
-				LandscapeIndices->Push((X + 0) + (Y + 1) * (ComponentSizeQuads + 1));
-				LandscapeIndices->Push((X + 1) + (Y + 1) * (ComponentSizeQuads + 1));
+				Indices.Push((X + 0) + (Y + 0) * (ComponentSizeQuads + 1));
+				Indices.Push((X + 0) + (Y + 1) * (ComponentSizeQuads + 1));
+				Indices.Push((X + 1) + (Y + 1) * (ComponentSizeQuads + 1));
 			}
 		}
 	}
@@ -571,7 +989,7 @@ void FGLTFDelayedLandscapeTask::Process()
 	JsonPrimitive.Attributes.Position = Builder.AddUniquePositionAccessor(PositionBuffer);
 	JsonPrimitive.Attributes.Normal = Builder.AddUniqueNormalAccessor(Normals);
 	JsonPrimitive.Attributes.TexCoords.AddUninitialized(1);
-	JsonPrimitive.Attributes.TexCoords[0] = Builder.AddUniqueUVAccessor(UVs);
-	JsonPrimitive.Indices = Builder.AddUniqueIndexAccessor(LandscapeIndices, JsonMesh->Name);
+	JsonPrimitive.Attributes.TexCoords[0] = Builder.AddUniqueUVAccessor(UV);
+	JsonPrimitive.Indices = Builder.AddUniqueIndexAccessor(Indices, JsonMesh->Name);
 	JsonPrimitive.Material = Builder.AddUniqueMaterial(&LandscapeMaterial);
 }
\ No newline at end of file
diff --git a/Engine/Plugins/Enterprise/GLTFExporter/Source/GLTFExporter/Private/Tasks/GLTFDelayedMeshTasks.h b/Engine/Plugins/Enterprise/GLTFExporter/Source/GLTFExporter/Private/Tasks/GLTFDelayedMeshTasks.h
--- a/Engine/Plugins/Enterprise/GLTFExporter/Source/GLTFExporter/Private/Tasks/GLTFDelayedMeshTasks.h
+++ b/Engine/Plugins/Enterprise/GLTFExporter/Source/GLTFExporter/Private/Tasks/GLTFDelayedMeshTasks.h
@@ -13,16 +13,18 @@
 class SkeletalMeshComponent;
 class StaticMeshComponent;
 
-class FGLTFDelayedStaticMeshTask : public FGLTFDelayedTask
+class FGLTFDelayedStaticAndSplineMeshTask : public FGLTFDelayedTask
 {
 public:
-
-	FGLTFDelayedStaticMeshTask(FGLTFConvertBuilder& Builder, FGLTFStaticMeshSectionConverter& MeshSectionConverter, const UStaticMesh* StaticMesh, const UStaticMeshComponent* StaticMeshComponent, FGLTFMaterialArray Materials, int32 LODIndex, FGLTFJsonMesh* JsonMesh)
+	FGLTFDelayedStaticAndSplineMeshTask(FGLTFConvertBuilder& Builder, FGLTFStaticMeshSectionConverter& MeshSectionConverter, const UStaticMesh* StaticMesh, 
+		const UStaticMeshComponent* StaticMeshComponent, const USplineMeshComponent* SplineMeshComponent,
+		FGLTFMaterialArray& Materials, int32 LODIndex, FGLTFJsonMesh* JsonMesh)
 		: FGLTFDelayedTask(EGLTFTaskPriority::Mesh)
 		, Builder(Builder)
 		, MeshSectionConverter(MeshSectionConverter)
 		, StaticMesh(StaticMesh)
 		, StaticMeshComponent(StaticMeshComponent)
+		, SplineMeshComponent(SplineMeshComponent)
 		, Materials(Materials)
 		, LODIndex(LODIndex)
 		, JsonMesh(JsonMesh)
@@ -34,58 +36,51 @@ public:
 	virtual void Process() override;
 
 private:
+#if WITH_EDITORONLY_DATA
+	void ProcessMeshDescription(const TArray<FStaticMaterial>& MaterialSlots, const FGLTFMeshData* MeshData);
+#endif
+	void ProcessRenderData(const TArray<FStaticMaterial>& MaterialSlots, const FGLTFMeshData* MeshData);
 
 	FGLTFConvertBuilder& Builder;
 	FGLTFStaticMeshSectionConverter& MeshSectionConverter;
 	const UStaticMesh* StaticMesh;
 	const UStaticMeshComponent* StaticMeshComponent;
-	const FGLTFMaterialArray Materials;
+	const USplineMeshComponent* SplineMeshComponent;
+	const FGLTFMaterialArray Materials; //We copy because the received value can be temporary value.
 	const int32 LODIndex;
 	FGLTFJsonMesh* JsonMesh;
 };
 
-class FGLTFDelayedSkeletalMeshTask : public FGLTFDelayedTask
+class FGLTFDelayedStaticMeshTask : public FGLTFDelayedStaticAndSplineMeshTask
 {
 public:
 
-	FGLTFDelayedSkeletalMeshTask(FGLTFConvertBuilder& Builder, FGLTFSkeletalMeshSectionConverter& MeshSectionConverter, const USkeletalMesh* SkeletalMesh, const USkeletalMeshComponent* SkeletalMeshComponent, FGLTFMaterialArray Materials, int32 LODIndex, FGLTFJsonMesh* JsonMesh)
-		: FGLTFDelayedTask(EGLTFTaskPriority::Mesh)
-		, Builder(Builder)
-		, MeshSectionConverter(MeshSectionConverter)
-		, SkeletalMesh(SkeletalMesh)
-		, SkeletalMeshComponent(SkeletalMeshComponent)
-		, Materials(Materials)
-		, LODIndex(LODIndex)
-		, JsonMesh(JsonMesh)
+	FGLTFDelayedStaticMeshTask(FGLTFConvertBuilder& Builder, FGLTFStaticMeshSectionConverter& MeshSectionConverter, const UStaticMesh* StaticMesh, const UStaticMeshComponent* StaticMeshComponent, FGLTFMaterialArray& Materials, int32 LODIndex, FGLTFJsonMesh* JsonMesh)
+		: FGLTFDelayedStaticAndSplineMeshTask(Builder, MeshSectionConverter, StaticMesh, StaticMeshComponent, nullptr, Materials, LODIndex, JsonMesh)
 	{
 	}
+};
 
-	virtual FString GetName() override;
-
-	virtual void Process() override;
-
-private:
+class FGLTFDelayedSplineMeshTask : public FGLTFDelayedStaticAndSplineMeshTask
+{
+public:
 
-	FGLTFConvertBuilder& Builder;
-	FGLTFSkeletalMeshSectionConverter& MeshSectionConverter;
-	const USkeletalMesh* SkeletalMesh;
-	const USkeletalMeshComponent* SkeletalMeshComponent;
-	const FGLTFMaterialArray Materials;
-	const int32 LODIndex;
-	FGLTFJsonMesh* JsonMesh;
+	FGLTFDelayedSplineMeshTask(FGLTFConvertBuilder& Builder, FGLTFStaticMeshSectionConverter& MeshSectionConverter, const UStaticMesh* StaticMesh, const USplineMeshComponent* SplineMeshComponent, FGLTFMaterialArray& Materials, int32 LODIndex, FGLTFJsonMesh* JsonMesh)
+		: FGLTFDelayedStaticAndSplineMeshTask(Builder, MeshSectionConverter, StaticMesh, nullptr, SplineMeshComponent, Materials, LODIndex, JsonMesh)
+	{
+	}
 };
 
-
-class FGLTFDelayedSplineMeshTask : public FGLTFDelayedTask
+class FGLTFDelayedSkeletalMeshTask : public FGLTFDelayedTask
 {
 public:
 
-	FGLTFDelayedSplineMeshTask(FGLTFConvertBuilder& Builder, FGLTFStaticMeshSectionConverter& MeshSectionConverter, const UStaticMesh& StaticMesh, const USplineMeshComponent& SplineMeshComponent, FGLTFMaterialArray Materials, int32 LODIndex, FGLTFJsonMesh* JsonMesh)
+	FGLTFDelayedSkeletalMeshTask(FGLTFConvertBuilder& Builder, FGLTFSkeletalMeshSectionConverter& MeshSectionConverter, const USkeletalMesh* SkeletalMesh, const USkeletalMeshComponent* SkeletalMeshComponent, FGLTFMaterialArray& Materials, int32 LODIndex, FGLTFJsonMesh* JsonMesh)
 		: FGLTFDelayedTask(EGLTFTaskPriority::Mesh)
 		, Builder(Builder)
 		, MeshSectionConverter(MeshSectionConverter)
-		, StaticMesh(StaticMesh)
-		, SplineMeshComponent(SplineMeshComponent)
+		, SkeletalMesh(SkeletalMesh)
+		, SkeletalMeshComponent(SkeletalMeshComponent)
 		, Materials(Materials)
 		, LODIndex(LODIndex)
 		, JsonMesh(JsonMesh)
@@ -97,12 +92,16 @@ public:
 	virtual void Process() override;
 
 private:
+#if WITH_EDITORONLY_DATA
+	void ProcessSourceModel(const TArray<FSkeletalMaterial>& MaterialSlots, const FGLTFMeshData* MeshData);
+#endif
+	void ProcessRenderData(const TArray<FSkeletalMaterial>& MaterialSlots, const FGLTFMeshData* MeshData);
 
 	FGLTFConvertBuilder& Builder;
-	FGLTFStaticMeshSectionConverter& MeshSectionConverter;
-	const UStaticMesh& StaticMesh;
-	const USplineMeshComponent& SplineMeshComponent;
-	const FGLTFMaterialArray Materials;
+	FGLTFSkeletalMeshSectionConverter& MeshSectionConverter;
+	const USkeletalMesh* SkeletalMesh;
+	const USkeletalMeshComponent* SkeletalMeshComponent;
+	const FGLTFMaterialArray Materials; //We copy because the received value can be temporary value.
 	const int32 LODIndex;
 	FGLTFJsonMesh* JsonMesh;
 };
diff --git a/Engine/Plugins/Enterprise/GLTFExporter/Source/GLTFExporter/Private/Utilities/GLTFCoreUtilities.cpp b/Engine/Plugins/Enterprise/GLTFExporter/Source/GLTFExporter/Private/Utilities/GLTFCoreUtilities.cpp
--- a/Engine/Plugins/Enterprise/GLTFExporter/Source/GLTFExporter/Private/Utilities/GLTFCoreUtilities.cpp
+++ b/Engine/Plugins/Enterprise/GLTFExporter/Source/GLTFExporter/Private/Utilities/GLTFCoreUtilities.cpp
@@ -51,6 +51,14 @@ FGLTFVector4 FGLTFCoreUtilities::ConvertTangent(const FVector3f& Tangent, const
 	return { SafeTangent.X, SafeTangent.Z, SafeTangent.Y, Normal.W >= 0 ? 1.0f : -1.0f };
 }
 
+FGLTFVector4 FGLTFCoreUtilities::ConvertTangent(const FVector4f& Tangent)
+{
+	// Unreal keeps the binormal sign in the normal's w-component.
+	// glTF keeps the binormal sign in the tangent's w-component.
+	const FVector3f SafeTangent = FVector3f(Tangent).GetSafeNormal();
+	return { SafeTangent.X, SafeTangent.Z, SafeTangent.Y, Tangent.W >= 0 ? 1.0f : -1.0f };
+}
+
 FGLTFInt8Vector4 FGLTFCoreUtilities::ConvertTangent(const FPackedNormal& Tangent, const FPackedNormal& Normal)
 {
 	const FPackedNormal SafeTangent = Tangent.ToFVector3f().GetSafeNormal();
diff --git a/Engine/Plugins/Enterprise/GLTFExporter/Source/GLTFExporter/Public/Builders/GLTFBuilder.h b/Engine/Plugins/Enterprise/GLTFExporter/Source/GLTFExporter/Public/Builders/GLTFBuilder.h
--- a/Engine/Plugins/Enterprise/GLTFExporter/Source/GLTFExporter/Public/Builders/GLTFBuilder.h
+++ b/Engine/Plugins/Enterprise/GLTFExporter/Source/GLTFExporter/Public/Builders/GLTFBuilder.h
@@ -11,6 +11,7 @@ class UStaticMesh;
 class UStaticMeshComponent;
 class USkeletalMesh;
 class USkeletalMeshComponent;
+class USplineMeshComponent;
 
 class GLTFEXPORTER_API FGLTFBuilder
 {
@@ -32,6 +33,7 @@ public:
 
 	int32 SanitizeLOD(const UStaticMesh* StaticMesh, const UStaticMeshComponent* StaticMeshComponent, int32 LODIndex) const;
 	int32 SanitizeLOD(const USkeletalMesh* SkeletalMesh, const USkeletalMeshComponent* SkeletalMeshComponent, int32 LODIndex) const;
+	int32 SanitizeLOD(const UStaticMesh* StaticMesh, const USplineMeshComponent* SplineMeshComponent, int32 LODIndex) const;
 
 private:
 
diff --git a/Engine/Plugins/Enterprise/GLTFExporter/Source/GLTFExporter/Public/Builders/GLTFConvertBuilder.h b/Engine/Plugins/Enterprise/GLTFExporter/Source/GLTFExporter/Public/Builders/GLTFConvertBuilder.h
--- a/Engine/Plugins/Enterprise/GLTFExporter/Source/GLTFExporter/Public/Builders/GLTFConvertBuilder.h
+++ b/Engine/Plugins/Enterprise/GLTFExporter/Source/GLTFExporter/Public/Builders/GLTFConvertBuilder.h
@@ -17,11 +17,10 @@
 #include "Converters/GLTFCameraConverters.h"
 #include "Converters/GLTFLightConverters.h"
 #include "Converters/GLTFMaterialVariantConverters.h"
+#include "Converters/GLTFMeshAttributesArray.h"
 
 class UMeshComponent;
 class UPropertyValue;
-class FGLTFNormalArray;
-class FGLTFUVArray;
 
 class GLTFEXPORTER_API FGLTFConvertBuilder : public FGLTFAnalyticsBuilder
 {
@@ -34,18 +33,22 @@ public:
 	bool IsSelectedActor(const AActor* Object) const;
 	bool IsRootActor(const AActor* Actor) const;
 
-	FGLTFJsonAccessor* AddUniquePositionAccessor(const FPositionVertexBuffer* VertexBuffer);
 	FGLTFJsonAccessor* AddUniquePositionAccessor(const FGLTFMeshSection* MeshSection, const FPositionVertexBuffer* VertexBuffer);
+	FGLTFJsonAccessor* AddUniquePositionAccessor(const FGLTFPositionArray& VertexBuffer);
 	FGLTFJsonAccessor* AddUniqueColorAccessor(const FGLTFMeshSection* MeshSection, const FColorVertexBuffer* VertexBuffer);
+	FGLTFJsonAccessor* AddUniqueColorAccessor(const FGLTFColorArray& VertexColorBuffer);
 	FGLTFJsonAccessor* AddUniqueNormalAccessor(const FGLTFMeshSection* MeshSection, const FStaticMeshVertexBuffer* VertexBuffer);
-	FGLTFJsonAccessor* AddUniqueNormalAccessor(const FGLTFNormalArray* Normals);
+	FGLTFJsonAccessor* AddUniqueNormalAccessor(const FGLTFNormalArray& Normals);
 	FGLTFJsonAccessor* AddUniqueTangentAccessor(const FGLTFMeshSection* MeshSection, const FStaticMeshVertexBuffer* VertexBuffer);
+	FGLTFJsonAccessor* AddUniqueTangentAccessor(const FGLTFTangentArray& Tangents);
 	FGLTFJsonAccessor* AddUniqueUVAccessor(const FGLTFMeshSection* MeshSection, const FStaticMeshVertexBuffer* VertexBuffer, int32 UVIndex);
-	FGLTFJsonAccessor* AddUniqueUVAccessor(const FGLTFUVArray* UVs); /* Supports single UV channel */
+	FGLTFJsonAccessor* AddUniqueUVAccessor(const FGLTFUVArray& UVs); /* Supports single UV channel */
 	FGLTFJsonAccessor* AddUniqueJointAccessor(const FGLTFMeshSection* MeshSection, const FSkinWeightVertexBuffer* VertexBuffer, int32 InfluenceOffset);
 	FGLTFJsonAccessor* AddUniqueWeightAccessor(const FGLTFMeshSection* MeshSection, const FSkinWeightVertexBuffer* VertexBuffer, int32 InfluenceOffset);
+	FGLTFJsonAccessor* AddUniqueJointAccessor(const FGLTFJointInfluenceArray& BoneIndices);
+	FGLTFJsonAccessor* AddUniqueWeightAccessor(const FGLTFJointWeightArray& Weights);
 	FGLTFJsonAccessor* AddUniqueIndexAccessor(const FGLTFMeshSection* MeshSection);
-	FGLTFJsonAccessor* AddUniqueIndexAccessor(const FGLTFIndexArray* IndexBuffer, const FString& MeshName);
+	FGLTFJsonAccessor* AddUniqueIndexAccessor(const FGLTFIndexArray& IndexBuffer, const FString& MeshName);
 
 	FGLTFJsonMesh* AddUniqueMesh(const UStaticMesh* StaticMesh, const FGLTFMaterialArray& Materials = {}, int32 LODIndex = INDEX_NONE);
 	FGLTFJsonMesh* AddUniqueMesh(const USkeletalMesh* SkeletalMesh, const FGLTFMaterialArray& Materials = {}, int32 LODIndex = INDEX_NONE);
@@ -100,13 +103,7 @@ public:
 	void RegisterObjectVariant(const UObject* Object, const UPropertyValue* Property);
 	const TArray<const UPropertyValue*>* GetObjectVariants(const UObject* Object) const;
 
-	TUniquePtr<IGLTFPositionBufferConverterRaw> PositionBufferConverterRaw = MakeUnique<FGLTFPositionBufferConverterRaw>(*this);
-	TUniquePtr<IGLTFNormalBufferConverterRaw> NormalBufferConverterRaw = MakeUnique<FGLTFNormalBufferConverterRaw>(*this);
-	TUniquePtr<IGLTFUVBufferConverterRaw> UVBufferConverterRaw = MakeUnique<FGLTFUVBufferConverterRaw>(*this);
-	TUniquePtr<IGLTFIndexBufferConverterRaw> IndexBufferConverterRaw = MakeUnique<FGLTFIndexBufferConverterRaw>(*this);
-	TUniquePtr<IGLTFSplineMeshConverter> SplineMeshConverter = MakeUnique<FGLTFSplineMeshConverter>(*this);
-	TUniquePtr<IGLTFLandscapeMeshConverter> LandscapeConverter = MakeUnique<FGLTFLandscapeMeshConverter>(*this);
-
+	//'Original' converters:
 	TUniquePtr<IGLTFPositionBufferConverter> PositionBufferConverter = MakeUnique<FGLTFPositionBufferConverter>(*this);
 	TUniquePtr<IGLTFColorBufferConverter> ColorBufferConverter = MakeUnique<FGLTFColorBufferConverter>(*this);
 	TUniquePtr<IGLTFNormalBufferConverter> NormalBufferConverter = MakeUnique<FGLTFNormalBufferConverter>(*this);
@@ -116,6 +113,20 @@ public:
 	TUniquePtr<IGLTFBoneWeightBufferConverter> BoneWeightBufferConverter = MakeUnique<FGLTFBoneWeightBufferConverter>(*this);
 	TUniquePtr<IGLTFIndexBufferConverter> IndexBufferConverter = MakeUnique<FGLTFIndexBufferConverter>(*this);
 
+	//Raw converters:
+	TUniquePtr<IGLTFPositionBufferConverterRaw> PositionBufferConverterRaw = MakeUnique<FGLTFPositionBufferConverterRaw>(*this);
+	TUniquePtr<IGLTFColorBufferConverterRaw> ColorBufferConverterRaw = MakeUnique<FGLTFColorBufferConverterRaw>(*this);
+	TUniquePtr<IGLTFNormalBufferConverterRaw> NormalBufferConverterRaw = MakeUnique<FGLTFNormalBufferConverterRaw>(*this);
+	TUniquePtr<IGLTFTangentBufferConverterRaw> TangentBufferConverterRaw = MakeUnique<FGLTFTangentBufferConverterRaw>(*this);
+	TUniquePtr<IGLTFUVBufferConverterRaw> UVBufferConverterRaw = MakeUnique<FGLTFUVBufferConverterRaw>(*this);
+	TUniquePtr<IGLTFIndexBufferConverterRaw> IndexBufferConverterRaw = MakeUnique<FGLTFIndexBufferConverterRaw>(*this);
+	TUniquePtr<IGLTFBoneIndexBufferConverterRaw> BoneIndexBufferConverterRaw = MakeUnique<FGLTFBoneIndexBufferConverterRaw>(*this);
+	TUniquePtr<IGLTFBoneWeightBufferConverterRaw> BoneWeightBufferConverterRaw = MakeUnique<FGLTFBoneWeightBufferConverterRaw>(*this);
+
+
+	TUniquePtr<IGLTFSplineMeshConverter> SplineMeshConverter = MakeUnique<FGLTFSplineMeshConverter>(*this);
+	TUniquePtr<IGLTFLandscapeMeshConverter> LandscapeConverter = MakeUnique<FGLTFLandscapeMeshConverter>(*this);
+
 	TUniquePtr<IGLTFStaticMeshConverter> StaticMeshConverter = MakeUnique<FGLTFStaticMeshConverter>(*this);
 	TUniquePtr<IGLTFSkeletalMeshConverter> SkeletalMeshConverter = MakeUnique<FGLTFSkeletalMeshConverter>(*this);
 
diff --git a/Engine/Plugins/Enterprise/GLTFExporter/Source/GLTFExporter/Public/Converters/GLTFAccessorConverters.h b/Engine/Plugins/Enterprise/GLTFExporter/Source/GLTFExporter/Public/Converters/GLTFAccessorConverters.h
--- a/Engine/Plugins/Enterprise/GLTFExporter/Source/GLTFExporter/Public/Converters/GLTFAccessorConverters.h
+++ b/Engine/Plugins/Enterprise/GLTFExporter/Source/GLTFExporter/Public/Converters/GLTFAccessorConverters.h
@@ -6,13 +6,12 @@
 #include "Converters/GLTFConverter.h"
 #include "Converters/GLTFBuilderContext.h"
 #include "Converters/GLTFMeshSection.h"
+#include "Converters/GLTFMeshAttributesArray.h"
 
 class FPositionVertexBuffer;
 class FColorVertexBuffer;
 class FStaticMeshVertexBuffer;
 class FSkinWeightVertexBuffer;
-class FGLTFNormalArray;
-class FGLTFUVArray;
 
 typedef TGLTFConverter<FGLTFJsonAccessor*, const FGLTFMeshSection*, const FPositionVertexBuffer*> IGLTFPositionBufferConverter;
 typedef TGLTFConverter<FGLTFJsonAccessor*, const FGLTFMeshSection*, const FColorVertexBuffer*> IGLTFColorBufferConverter;
@@ -22,10 +21,15 @@ typedef TGLTFConverter<FGLTFJsonAccessor*, const FGLTFMeshSection*, const FStati
 typedef TGLTFConverter<FGLTFJsonAccessor*, const FGLTFMeshSection*, const FSkinWeightVertexBuffer*, uint32> IGLTFBoneIndexBufferConverter;
 typedef TGLTFConverter<FGLTFJsonAccessor*, const FGLTFMeshSection*, const FSkinWeightVertexBuffer*, uint32> IGLTFBoneWeightBufferConverter;
 typedef TGLTFConverter<FGLTFJsonAccessor*, const FGLTFMeshSection*> IGLTFIndexBufferConverter;
-typedef TGLTFConverter<FGLTFJsonAccessor*, const FPositionVertexBuffer*> IGLTFPositionBufferConverterRaw;
-typedef TGLTFConverter<FGLTFJsonAccessor*, const FGLTFIndexArray*, FString> IGLTFIndexBufferConverterRaw;
-typedef TGLTFConverter<FGLTFJsonAccessor*, const FGLTFNormalArray*> IGLTFNormalBufferConverterRaw;
-typedef TGLTFConverter<FGLTFJsonAccessor*, const FGLTFUVArray*> IGLTFUVBufferConverterRaw;
+
+typedef TGLTFConverter<FGLTFJsonAccessor*, FGLTFPositionArray> IGLTFPositionBufferConverterRaw;
+typedef TGLTFConverter<FGLTFJsonAccessor*, FGLTFIndexArray, FString> IGLTFIndexBufferConverterRaw;
+typedef TGLTFConverter<FGLTFJsonAccessor*, FGLTFNormalArray> IGLTFNormalBufferConverterRaw;
+typedef TGLTFConverter<FGLTFJsonAccessor*, FGLTFUVArray> IGLTFUVBufferConverterRaw;
+typedef TGLTFConverter<FGLTFJsonAccessor*, FGLTFColorArray> IGLTFColorBufferConverterRaw;
+typedef TGLTFConverter<FGLTFJsonAccessor*, FGLTFTangentArray> IGLTFTangentBufferConverterRaw;
+typedef TGLTFConverter<FGLTFJsonAccessor*, FGLTFJointInfluenceArray> IGLTFBoneIndexBufferConverterRaw;
+typedef TGLTFConverter<FGLTFJsonAccessor*, FGLTFJointWeightArray> IGLTFBoneWeightBufferConverterRaw;
 
 
 class GLTFEXPORTER_API FGLTFPositionBufferConverter : public FGLTFBuilderContext, public IGLTFPositionBufferConverter
@@ -154,50 +158,72 @@ protected:
 class GLTFEXPORTER_API FGLTFPositionBufferConverterRaw : public FGLTFBuilderContext, public IGLTFPositionBufferConverterRaw
 {
 public:
-
 	using FGLTFBuilderContext::FGLTFBuilderContext;
 
 protected:
-
-	virtual FGLTFJsonAccessor* Convert(const FPositionVertexBuffer* VertexBuffer) override;
+	virtual FGLTFJsonAccessor* Convert(FGLTFPositionArray VertexBuffer) override;
 };
 
 class GLTFEXPORTER_API FGLTFIndexBufferConverterRaw : public FGLTFBuilderContext, public IGLTFIndexBufferConverterRaw
 {
 public:
-
 	using FGLTFBuilderContext::FGLTFBuilderContext;
 
 protected:
-	virtual FGLTFJsonAccessor* Convert(const FGLTFIndexArray* IndexBuffer, FString MeshName) override;
+	virtual FGLTFJsonAccessor* Convert(FGLTFIndexArray IndexBuffer, FString MeshName) override;
 };
 
 class GLTFEXPORTER_API FGLTFNormalBufferConverterRaw : public FGLTFBuilderContext, public IGLTFNormalBufferConverterRaw
 {
 public:
-
 	using FGLTFBuilderContext::FGLTFBuilderContext;
 
 protected:
+	virtual FGLTFJsonAccessor* Convert(FGLTFNormalArray NormalsSource) override;
+};
 
-	virtual FGLTFJsonAccessor* Convert(const FGLTFNormalArray* NormalSource) override;
+class GLTFEXPORTER_API FGLTFUVBufferConverterRaw : public FGLTFBuilderContext, public IGLTFUVBufferConverterRaw
+{
+public:
 
-private:
+	using FGLTFBuilderContext::FGLTFBuilderContext;
 
-	FGLTFJsonBufferView* ConvertBufferView(const FGLTFNormalArray* NormalSource);
+protected:
+	virtual FGLTFJsonAccessor* Convert(FGLTFUVArray UVSource) override;
 };
 
-class GLTFEXPORTER_API FGLTFUVBufferConverterRaw : public FGLTFBuilderContext, public IGLTFUVBufferConverterRaw
+class GLTFEXPORTER_API FGLTFColorBufferConverterRaw : public FGLTFBuilderContext, public IGLTFColorBufferConverterRaw
 {
 public:
+	using FGLTFBuilderContext::FGLTFBuilderContext;
 
+protected:
+	virtual FGLTFJsonAccessor* Convert(FGLTFColorArray VertexColorBuffer) override;
+};
+
+class GLTFEXPORTER_API FGLTFTangentBufferConverterRaw : public FGLTFBuilderContext, public IGLTFTangentBufferConverterRaw
+{
+public:
 	using FGLTFBuilderContext::FGLTFBuilderContext;
 
 protected:
+	virtual FGLTFJsonAccessor* Convert(FGLTFTangentArray TangentSource) override;
+};
 
-	virtual FGLTFJsonAccessor* Convert(const FGLTFUVArray* UVSource) override;
+class GLTFEXPORTER_API FGLTFBoneIndexBufferConverterRaw : public FGLTFBuilderContext, public IGLTFBoneIndexBufferConverterRaw
+{
+public:
+	using FGLTFBuilderContext::FGLTFBuilderContext;
 
-private:
+protected:
+	virtual FGLTFJsonAccessor* Convert(FGLTFJointInfluenceArray Joints) override;
+};
 
-	FGLTFJsonBufferView* ConvertBufferView(const FGLTFUVArray* UVSource);
+class GLTFEXPORTER_API FGLTFBoneWeightBufferConverterRaw : public FGLTFBuilderContext, public IGLTFBoneWeightBufferConverterRaw
+{
+public:
+	using FGLTFBuilderContext::FGLTFBuilderContext;
+
+protected:
+	virtual FGLTFJsonAccessor* Convert(FGLTFJointWeightArray Weights) override;
 };
\ No newline at end of file
diff --git a/Engine/Plugins/Enterprise/GLTFExporter/Source/GLTFExporter/Public/Converters/GLTFIndexArray.h b/Engine/Plugins/Enterprise/GLTFExporter/Source/GLTFExporter/Public/Converters/GLTFIndexArray.h
deleted file mode 100644
diff --git a/Engine/Plugins/Enterprise/GLTFExporter/Source/GLTFExporter/Public/Converters/GLTFMaterialConverters.h b/Engine/Plugins/Enterprise/GLTFExporter/Source/GLTFExporter/Public/Converters/GLTFMaterialConverters.h
--- a/Engine/Plugins/Enterprise/GLTFExporter/Source/GLTFExporter/Public/Converters/GLTFMaterialConverters.h
+++ b/Engine/Plugins/Enterprise/GLTFExporter/Source/GLTFExporter/Public/Converters/GLTFMaterialConverters.h
@@ -6,7 +6,7 @@
 #include "Converters/GLTFConverter.h"
 #include "Converters/GLTFBuilderContext.h"
 #include "Converters/GLTFMeshData.h"
-#include "Converters/GLTFIndexArray.h"
+#include "Converters/GLTFMeshAttributesArray.h"
 #include "Converters/GLTFUVOverlapChecker.h"
 #include "Converters/GLTFUVBoundsCalculator.h"
 #include "Converters/GLTFUVDegenerateChecker.h"
diff --git a/Engine/Plugins/Enterprise/GLTFExporter/Source/GLTFExporter/Public/Converters/GLTFMeshAttributesArray.h b/Engine/Plugins/Enterprise/GLTFExporter/Source/GLTFExporter/Public/Converters/GLTFMeshAttributesArray.h
new file mode 100644
--- /dev/null
+++ b/Engine/Plugins/Enterprise/GLTFExporter/Source/GLTFExporter/Public/Converters/GLTFMeshAttributesArray.h
@@ -0,0 +1,34 @@
+// Copyright Epic Games, Inc. All Rights Reserved.
+
+#pragma once
+
+#include "CoreMinimal.h"
+#include "BoneIndices.h"
+
+template <typename T>
+class FGLTFMeshAttributesArray : public TArray<T>
+{
+public:
+
+	using TArray<T>::TArray;
+	using TArray<T>::operator=;
+
+	friend uint32 GetTypeHash(const FGLTFMeshAttributesArray& AttributesArray)
+	{
+		uint32 Hash = GetTypeHash(AttributesArray.Num());
+		for (const T& Attribute : AttributesArray)
+		{
+			Hash = HashCombine(Hash, GetTypeHash(Attribute));
+		}
+		return Hash;
+	}
+};
+
+typedef FGLTFMeshAttributesArray<FColor> FGLTFColorArray;
+typedef FGLTFMeshAttributesArray<int32> FGLTFIndexArray;
+typedef FGLTFMeshAttributesArray<UE::Math::TIntVector4<FBoneIndexType>> FGLTFJointInfluenceArray;
+typedef FGLTFMeshAttributesArray<UE::Math::TIntVector4<uint16>> FGLTFJointWeightArray;
+typedef FGLTFMeshAttributesArray<FVector3f> FGLTFNormalArray;
+typedef FGLTFMeshAttributesArray<FVector3f> FGLTFPositionArray;
+typedef FGLTFMeshAttributesArray<FVector4f> FGLTFTangentArray;
+typedef FGLTFMeshAttributesArray<FVector2f> FGLTFUVArray;
\ No newline at end of file
diff --git a/Engine/Plugins/Enterprise/GLTFExporter/Source/GLTFExporter/Public/Converters/GLTFMeshSection.h b/Engine/Plugins/Enterprise/GLTFExporter/Source/GLTFExporter/Public/Converters/GLTFMeshSection.h
--- a/Engine/Plugins/Enterprise/GLTFExporter/Source/GLTFExporter/Public/Converters/GLTFMeshSection.h
+++ b/Engine/Plugins/Enterprise/GLTFExporter/Source/GLTFExporter/Public/Converters/GLTFMeshSection.h
@@ -4,7 +4,7 @@
 
 #include "CoreMinimal.h"
 #include "BoneIndices.h"
-#include "Converters/GLTFIndexArray.h"
+#include "Converters/GLTFMeshAttributesArray.h"
 
 class UStaticMesh;
 class USkeletalMesh;
diff --git a/Engine/Plugins/Enterprise/GLTFExporter/Source/GLTFExporter/Public/Converters/GLTFMeshSectionConverters.h b/Engine/Plugins/Enterprise/GLTFExporter/Source/GLTFExporter/Public/Converters/GLTFMeshSectionConverters.h
--- a/Engine/Plugins/Enterprise/GLTFExporter/Source/GLTFExporter/Public/Converters/GLTFMeshSectionConverters.h
+++ b/Engine/Plugins/Enterprise/GLTFExporter/Source/GLTFExporter/Public/Converters/GLTFMeshSectionConverters.h
@@ -4,7 +4,7 @@
 
 #include "Converters/GLTFConverter.h"
 #include "Converters/GLTFMeshSection.h"
-#include "Converters/GLTFIndexArray.h"
+#include "Converters/GLTFMeshAttributesArray.h"
 
 template <typename MeshType>
 class TGLTFMeshSectionConverter : public TGLTFConverter<const FGLTFMeshSection*, const MeshType*, int32, FGLTFIndexArray>
diff --git a/Engine/Plugins/Enterprise/GLTFExporter/Source/GLTFExporter/Public/Converters/GLTFNormalArray.h b/Engine/Plugins/Enterprise/GLTFExporter/Source/GLTFExporter/Public/Converters/GLTFNormalArray.h
deleted file mode 100644
diff --git a/Engine/Plugins/Enterprise/GLTFExporter/Source/GLTFExporter/Public/Converters/GLTFUVArray.h b/Engine/Plugins/Enterprise/GLTFExporter/Source/GLTFExporter/Public/Converters/GLTFUVArray.h
deleted file mode 100644
diff --git a/Engine/Plugins/Enterprise/GLTFExporter/Source/GLTFExporter/Public/Converters/GLTFUVBoundsCalculator.h b/Engine/Plugins/Enterprise/GLTFExporter/Source/GLTFExporter/Public/Converters/GLTFUVBoundsCalculator.h
--- a/Engine/Plugins/Enterprise/GLTFExporter/Source/GLTFExporter/Public/Converters/GLTFUVBoundsCalculator.h
+++ b/Engine/Plugins/Enterprise/GLTFExporter/Source/GLTFExporter/Public/Converters/GLTFUVBoundsCalculator.h
@@ -3,7 +3,7 @@
 #pragma once
 
 #include "Converters/GLTFConverter.h"
-#include "Converters/GLTFIndexArray.h"
+#include "Converters/GLTFMeshAttributesArray.h"
 
 struct FMeshDescription;
 
diff --git a/Engine/Plugins/Enterprise/GLTFExporter/Source/GLTFExporter/Public/Converters/GLTFUVDegenerateChecker.h b/Engine/Plugins/Enterprise/GLTFExporter/Source/GLTFExporter/Public/Converters/GLTFUVDegenerateChecker.h
--- a/Engine/Plugins/Enterprise/GLTFExporter/Source/GLTFExporter/Public/Converters/GLTFUVDegenerateChecker.h
+++ b/Engine/Plugins/Enterprise/GLTFExporter/Source/GLTFExporter/Public/Converters/GLTFUVDegenerateChecker.h
@@ -3,7 +3,7 @@
 #pragma once
 
 #include "Converters/GLTFConverter.h"
-#include "Converters/GLTFIndexArray.h"
+#include "Converters/GLTFMeshAttributesArray.h"
 #include "Containers/StaticArray.h"
 
 struct FMeshDescription;
diff --git a/Engine/Plugins/Enterprise/GLTFExporter/Source/GLTFExporter/Public/Converters/GLTFUVOverlapChecker.h b/Engine/Plugins/Enterprise/GLTFExporter/Source/GLTFExporter/Public/Converters/GLTFUVOverlapChecker.h
--- a/Engine/Plugins/Enterprise/GLTFExporter/Source/GLTFExporter/Public/Converters/GLTFUVOverlapChecker.h
+++ b/Engine/Plugins/Enterprise/GLTFExporter/Source/GLTFExporter/Public/Converters/GLTFUVOverlapChecker.h
@@ -3,7 +3,7 @@
 #pragma once
 
 #include "Converters/GLTFConverter.h"
-#include "Converters/GLTFIndexArray.h"
+#include "Converters/GLTFMeshAttributesArray.h"
 
 class UMaterialInterface;
 struct FMeshDescription;
diff --git a/Engine/Plugins/Enterprise/GLTFExporter/Source/GLTFExporter/Public/Json/GLTFJsonMesh.h b/Engine/Plugins/Enterprise/GLTFExporter/Source/GLTFExporter/Public/Json/GLTFJsonMesh.h
--- a/Engine/Plugins/Enterprise/GLTFExporter/Source/GLTFExporter/Public/Json/GLTFJsonMesh.h
+++ b/Engine/Plugins/Enterprise/GLTFExporter/Source/GLTFExporter/Public/Json/GLTFJsonMesh.h
@@ -25,6 +25,8 @@ struct GLTFEXPORTER_API FGLTFJsonAttributes : IGLTFJsonObject
 	}
 
 	virtual void WriteObject(IGLTFJsonWriter& Writer) const override;
+
+	bool HasValue() const;
 };
 
 struct GLTFEXPORTER_API FGLTFJsonPrimitive : IGLTFJsonObject
@@ -44,6 +46,8 @@ struct GLTFEXPORTER_API FGLTFJsonPrimitive : IGLTFJsonObject
 	}
 
 	virtual void WriteObject(IGLTFJsonWriter& Writer) const override;
+
+	bool HasValue() const;
 };
 
 struct GLTFEXPORTER_API FGLTFJsonMesh : IGLTFJsonIndexedObject
@@ -54,6 +58,8 @@ struct GLTFEXPORTER_API FGLTFJsonMesh : IGLTFJsonIndexedObject
 
 	virtual void WriteObject(IGLTFJsonWriter& Writer) const override;
 
+	bool HasValue() const;
+
 protected:
 
 	friend TGLTFJsonIndexedObjectArray<FGLTFJsonMesh, void>;
diff --git a/Engine/Plugins/Enterprise/GLTFExporter/Source/GLTFExporter/Public/Options/GLTFExportOptions.h b/Engine/Plugins/Enterprise/GLTFExporter/Source/GLTFExporter/Public/Options/GLTFExportOptions.h
--- a/Engine/Plugins/Enterprise/GLTFExporter/Source/GLTFExporter/Public/Options/GLTFExportOptions.h
+++ b/Engine/Plugins/Enterprise/GLTFExporter/Source/GLTFExporter/Public/Options/GLTFExportOptions.h
@@ -105,6 +105,10 @@ class GLTFEXPORTER_API UGLTFExportOptions : public UObject
 	UPROPERTY(EditAnywhere, BlueprintReadWrite, Config, Category = Mesh, Meta = (ClampMin = "0"))
 	int32 DefaultLevelOfDetail;
 
+	/** If enabled, exports the SourceModel. If false, exports the Render Data. */
+	UPROPERTY(EditAnywhere, BlueprintReadWrite, Config, Category = Mesh)
+	bool bExportSourceModel;
+
 	/** If enabled, export vertex color. Not recommended due to vertex colors always being used as a base color multiplier in glTF, regardless of material. Often producing undesirable results. */
 	UPROPERTY(EditAnywhere, BlueprintReadWrite, Config, Category = Mesh)
 	bool bExportVertexColors;
diff --git a/Engine/Plugins/Enterprise/GLTFExporter/Source/GLTFExporter/Public/Utilities/GLTFCoreUtilities.h b/Engine/Plugins/Enterprise/GLTFExporter/Source/GLTFExporter/Public/Utilities/GLTFCoreUtilities.h
--- a/Engine/Plugins/Enterprise/GLTFExporter/Source/GLTFExporter/Public/Utilities/GLTFCoreUtilities.h
+++ b/Engine/Plugins/Enterprise/GLTFExporter/Source/GLTFExporter/Public/Utilities/GLTFCoreUtilities.h
@@ -30,6 +30,7 @@ struct GLTFEXPORTER_API FGLTFCoreUtilities
 	static FGLTFInt16Vector4 ConvertNormal(const FPackedRGBA16N& Normal);
 
 	static FGLTFVector4 ConvertTangent(const FVector3f& Tangent, const FVector4f& Normal = FVector4f(ForceInitToZero));
+	static FGLTFVector4 ConvertTangent(const FVector4f& Tangent);
 	static FGLTFInt8Vector4 ConvertTangent(const FPackedNormal& Tangent, const FPackedNormal& Normal = FPackedNormal());
 	static FGLTFInt16Vector4 ConvertTangent(const FPackedRGBA16N& Tangent, const FPackedRGBA16N& Normal = FPackedRGBA16N());
 
