[glTF Exporter] Fixing Landscape exports and fixing empty primitive/mesh usages.

#rnx

diff --git a/Engine/Plugins/Enterprise/GLTFExporter/Source/GLTFExporter/Private/Builders/GLTFJsonBuilder.cpp b/Engine/Plugins/Enterprise/GLTFExporter/Source/GLTFExporter/Private/Builders/GLTFJsonBuilder.cpp
--- a/Engine/Plugins/Enterprise/GLTFExporter/Source/GLTFExporter/Private/Builders/GLTFJsonBuilder.cpp
+++ b/Engine/Plugins/Enterprise/GLTFExporter/Source/GLTFExporter/Private/Builders/GLTFJsonBuilder.cpp
@@ -18,6 +18,8 @@ FGLTFJsonBuilder::FGLTFJsonBuilder(const FString& FileName, const UGLTFExportOpt
 
 bool FGLTFJsonBuilder::WriteJsonArchive(FArchive& Archive)
 {
+	ValidateAndFixGLTFJson();
+
 	JsonRoot.WriteJson(Archive, !bIsGLB, ExportOptions->bSkipNearDefaultValues ? KINDA_SMALL_NUMBER : 0);
 	return true;
 }
@@ -120,3 +122,50 @@ FString FGLTFJsonBuilder::GetCopyrightString()
 {
 	return GetDefault<UGeneralProjectSettings>()->CopyrightNotice;
 }
+
+
+void FGLTFJsonBuilder::ValidateAndFixGLTFJson()
+{
+	//Sort through Meshes, remove empty ones (without indices/attributes/primitives)
+	bool bRemovedElement = false;
+	TSet<int32> RemoveOriginalIndices;
+	for (int32 MeshIndex = 0, OriginalMeshIndex = 0; MeshIndex < JsonRoot.Meshes.Num(); MeshIndex++, OriginalMeshIndex++)
+	{
+		if (!JsonRoot.Meshes[MeshIndex]->HasValue())
+		{
+			RemoveOriginalIndices.Add(OriginalMeshIndex);
+			bRemovedElement = true;
+		}
+	}
+
+	if (bRemovedElement)
+	{
+		for (size_t NodeIndex = 0; NodeIndex < JsonRoot.Nodes.Num(); NodeIndex++)
+		{
+			if (JsonRoot.Nodes[NodeIndex]->Mesh)
+			{
+				if (RemoveOriginalIndices.Contains(JsonRoot.Nodes[NodeIndex]->Mesh->Index))
+				{
+					JsonRoot.Nodes[NodeIndex]->Mesh = nullptr;
+				}
+			}
+		}
+
+		TArray<int32> RemoveOriginalIndicesSorted = RemoveOriginalIndices.Array();
+		RemoveOriginalIndicesSorted.Sort([](int32 A, int32 B) {
+			return A > B;
+			});
+
+		for (int32 RemoveMeshIndex : RemoveOriginalIndicesSorted)
+		{
+			JsonRoot.Meshes.Remove(RemoveMeshIndex);
+		}
+
+		if (bRemovedElement)
+		{
+			JsonRoot.Meshes.FixElementIndices();
+		}
+	}
+
+	
+}
\ No newline at end of file
diff --git a/Engine/Plugins/Enterprise/GLTFExporter/Source/GLTFExporter/Private/Converters/GLTFNodeConverters.cpp b/Engine/Plugins/Enterprise/GLTFExporter/Source/GLTFExporter/Private/Converters/GLTFNodeConverters.cpp
--- a/Engine/Plugins/Enterprise/GLTFExporter/Source/GLTFExporter/Private/Converters/GLTFNodeConverters.cpp
+++ b/Engine/Plugins/Enterprise/GLTFExporter/Source/GLTFExporter/Private/Converters/GLTFNodeConverters.cpp
@@ -15,6 +15,7 @@
 #include "Engine/SkeletalMeshSocket.h"
 #include "GameFramework/Pawn.h"
 
+#include "LandscapeStreamingProxy.h"
 #include "LandscapeComponent.h"
 #include "GLTFMeshUtilities.h"
 #include "Components/InstancedStaticMeshComponent.h"
@@ -38,6 +39,13 @@ FGLTFJsonNode* FGLTFActorConverter::Convert(const AActor* Actor)
 	}
 #endif
 
+	if (const ALandscapeStreamingProxy* LevelInstanceEditorInstanceActor = Cast<ALandscapeStreamingProxy>(Actor))
+	{
+		//The actual object that it is proxied to will be iterated through as well. To avoid duplications we won't include the Proxy.
+		return nullptr;
+	}
+	
+
 	const USceneComponent* RootComponent = Actor->GetRootComponent();
 	FGLTFJsonNode* RootNode = Builder.AddUniqueNode(RootComponent);
 
@@ -90,6 +98,12 @@ FGLTFJsonNode* FGLTFComponentConverter::Convert(const USceneComponent* SceneComp
 	}
 #endif
 
+	if (const ALandscapeStreamingProxy* LevelInstanceEditorInstanceActor = Cast<ALandscapeStreamingProxy>(Owner))
+	{
+		//The actual object that it is proxied to will be iterated through as well. To avoid duplications we won't include the Proxy.
+		return nullptr;
+	}
+
 	if (!Builder.IsSelectedActor(Owner))
 	{
 		return nullptr;
diff --git a/Engine/Plugins/Enterprise/GLTFExporter/Source/GLTFExporter/Private/Json/GLTFJsonMesh.cpp b/Engine/Plugins/Enterprise/GLTFExporter/Source/GLTFExporter/Private/Json/GLTFJsonMesh.cpp
--- a/Engine/Plugins/Enterprise/GLTFExporter/Source/GLTFExporter/Private/Json/GLTFJsonMesh.cpp
+++ b/Engine/Plugins/Enterprise/GLTFExporter/Source/GLTFExporter/Private/Json/GLTFJsonMesh.cpp
@@ -92,19 +92,16 @@ void FGLTFJsonMesh::WriteObject(IGLTFJsonWriter& Writer) const
 		Writer.Write(TEXT("name"), Name);
 	}
 
-	if (HasValue())
+	Writer.SetIdentifier(TEXT("primitives"));
+	Writer.StartArray();
+	for (const FGLTFJsonPrimitive& Primitive : Primitives)
 	{
-		Writer.SetIdentifier(TEXT("primitives"));
-		Writer.StartArray();
-		for (const FGLTFJsonPrimitive& Primitive : Primitives)
+		if (Primitive.HasValue())
 		{
-			if (Primitive.HasValue())
-			{
-				Writer.Write(Primitive);
-			}
+			Writer.Write(Primitive);
 		}
-		Writer.EndArray();
 	}
+	Writer.EndArray();
 }
 
 bool FGLTFJsonMesh::HasValue() const
diff --git a/Engine/Plugins/Enterprise/GLTFExporter/Source/GLTFExporter/Private/Tasks/GLTFDelayedMeshTasks.cpp b/Engine/Plugins/Enterprise/GLTFExporter/Source/GLTFExporter/Private/Tasks/GLTFDelayedMeshTasks.cpp
--- a/Engine/Plugins/Enterprise/GLTFExporter/Source/GLTFExporter/Private/Tasks/GLTFDelayedMeshTasks.cpp
+++ b/Engine/Plugins/Enterprise/GLTFExporter/Source/GLTFExporter/Private/Tasks/GLTFDelayedMeshTasks.cpp
@@ -942,7 +942,6 @@ void FGLTFDelayedLandscapeTask::Process()
 			CDI.GetWeightmapTextureData(AllocInfo.LayerInfo, CompVisData);
 		}
 	}
-	CompVisData.AddZeroed((ComponentSizeQuads + 1) * (ComponentSizeQuads + 1));
 
 	if (CompVisData.Num() > 0)
 	{
@@ -986,10 +985,13 @@ void FGLTFDelayedLandscapeTask::Process()
 		}
 	}
 
-	JsonPrimitive.Attributes.Position = Builder.AddUniquePositionAccessor(PositionBuffer);
-	JsonPrimitive.Attributes.Normal = Builder.AddUniqueNormalAccessor(Normals);
-	JsonPrimitive.Attributes.TexCoords.AddUninitialized(1);
-	JsonPrimitive.Attributes.TexCoords[0] = Builder.AddUniqueUVAccessor(UV);
-	JsonPrimitive.Indices = Builder.AddUniqueIndexAccessor(Indices, JsonMesh->Name);
-	JsonPrimitive.Material = Builder.AddUniqueMaterial(&LandscapeMaterial);
+	if (Indices.Num())
+	{
+		JsonPrimitive.Attributes.Position = Builder.AddUniquePositionAccessor(PositionBuffer);
+		JsonPrimitive.Attributes.Normal = Builder.AddUniqueNormalAccessor(Normals);
+		JsonPrimitive.Attributes.TexCoords.AddUninitialized(1);
+		JsonPrimitive.Attributes.TexCoords[0] = Builder.AddUniqueUVAccessor(UV);
+		JsonPrimitive.Indices = Builder.AddUniqueIndexAccessor(Indices, JsonMesh->Name);
+		JsonPrimitive.Material = Builder.AddUniqueMaterial(&LandscapeMaterial);
+	}
 }
\ No newline at end of file
diff --git a/Engine/Plugins/Enterprise/GLTFExporter/Source/GLTFExporter/Private/Utilities/GLTFLandscapeComponentDataInterface.cpp b/Engine/Plugins/Enterprise/GLTFExporter/Source/GLTFExporter/Private/Utilities/GLTFLandscapeComponentDataInterface.cpp
--- a/Engine/Plugins/Enterprise/GLTFExporter/Source/GLTFExporter/Private/Utilities/GLTFLandscapeComponentDataInterface.cpp
+++ b/Engine/Plugins/Enterprise/GLTFExporter/Source/GLTFExporter/Private/Utilities/GLTFLandscapeComponentDataInterface.cpp
@@ -154,6 +154,13 @@ bool FGLTFLandscapeComponentDataInterface::GetWeightmapTextureData(ULandscapeLay
 		return false;
 	}
 
+	UTexture2D* WeightMapTexture = ComponentWeightmapTextures[ComponentWeightmapLayerAllocations[LayerIdx].WeightmapTextureIndex];
+
+	if (!WeightMapTexture)
+	{
+		return false;
+	}
+
 	// If requested to skip the duplicate row/col of texture data
 	int32 WeightmapSize = bInRemoveSubsectionDuplicates ?
 		((ComponentSubsectionSizeQuads * Component.NumSubsections) + 1) >> MipLevel :
@@ -161,14 +168,30 @@ bool FGLTFLandscapeComponentDataInterface::GetWeightmapTextureData(ULandscapeLay
 
 	OutData.Empty(FMath::Square(WeightmapSize));
 	OutData.AddUninitialized(FMath::Square(WeightmapSize));
+	
+#if WITH_EDITOR
+	if (MipLevel >= WeightMapTexture->Source.GetNumMips())
+	{
+		return false;
+	}
+	TArray64<uint8> MipData;
+	WeightMapTexture->Source.GetMipData(MipData, MipLevel);
+	const FColor* WeightMipData = (FColor*)MipData.GetData();
+#else
+	FTexturePlatformData* PlatformData = WeightMapTexture->GetPlatformData();
+	if (MipLevel >= PlatformData->Mips.Num())
+	{
+		return false;
+	}
+	void* MipData = nullptr;
+	PlatformData->Mips[MipLevel].BulkData.GetCopy(&MipData, false);
+	const FColor* WeightMipData = (FColor*)MipData;
+#endif
 
-	// DataInterface Lock is a LockMipReadOnly on the texture
-	const TIndirectArray<FTexture2DMipMap>& WeightMips = ComponentWeightmapTextures[ComponentWeightmapLayerAllocations[LayerIdx].WeightmapTextureIndex]->GetPlatformMips();
-	if (MipLevel >= WeightMips.Num())
+	if (!WeightMipData)
 	{
 		return false;
 	}
-	const FColor* WeightMipData = (FColor*)WeightMips[MipLevel].BulkData.LockReadOnly();
 
 	// Channel remapping
 	int32 ChannelOffsets[4] = { (int32)STRUCT_OFFSET(FColor, R), (int32)STRUCT_OFFSET(FColor, G), (int32)STRUCT_OFFSET(FColor, B), (int32)STRUCT_OFFSET(FColor, A) };
@@ -181,8 +204,6 @@ bool FGLTFLandscapeComponentDataInterface::GetWeightmapTextureData(ULandscapeLay
 		OutData[i] = bInRemoveSubsectionDuplicates ? SrcTextureData[VertexIndexToTexel(i) * sizeof(FColor)] : SrcTextureData[i * sizeof(FColor)];
 	}
 
-	WeightMips[MipLevel].BulkData.Unlock();
-
 	return true;
 }
 
diff --git a/Engine/Plugins/Enterprise/GLTFExporter/Source/GLTFExporter/Public/Builders/GLTFJsonBuilder.h b/Engine/Plugins/Enterprise/GLTFExporter/Source/GLTFExporter/Public/Builders/GLTFJsonBuilder.h
--- a/Engine/Plugins/Enterprise/GLTFExporter/Source/GLTFExporter/Public/Builders/GLTFJsonBuilder.h
+++ b/Engine/Plugins/Enterprise/GLTFExporter/Source/GLTFExporter/Public/Builders/GLTFJsonBuilder.h
@@ -43,4 +43,6 @@ private:
 	static FString GetCopyrightString();
 
 	FGLTFJsonRoot JsonRoot;
+
+	void ValidateAndFixGLTFJson();
 };
diff --git a/Engine/Plugins/Enterprise/GLTFExporter/Source/GLTFExporter/Public/Json/GLTFJsonIndexedObjectArray.h b/Engine/Plugins/Enterprise/GLTFExporter/Source/GLTFExporter/Public/Json/GLTFJsonIndexedObjectArray.h
--- a/Engine/Plugins/Enterprise/GLTFExporter/Source/GLTFExporter/Public/Json/GLTFJsonIndexedObjectArray.h
+++ b/Engine/Plugins/Enterprise/GLTFExporter/Source/GLTFExporter/Public/Json/GLTFJsonIndexedObjectArray.h
@@ -64,6 +64,20 @@ struct TGLTFJsonIndexedObjectArray : IGLTFJsonArray
 		return Element;
 	}
 
+	//Do not forget to call FixElementIndices after.
+	void Remove(int32 Index)
+	{
+		Array.RemoveAt(Index);
+	}
+
+	void FixElementIndices()
+	{
+		for (size_t ElementIndex = 0; ElementIndex < Array.Num(); ElementIndex++)
+		{
+			Array[ElementIndex]->Index = ElementIndex;
+		}
+	}
+
 	bool IsValidIndex(SizeType Index) const
 	{
 		return Array.IsValidIndex(Index);
diff --git a/Engine/Plugins/Enterprise/GLTFExporter/Source/GLTFExporter/Public/Json/GLTFJsonRoot.h b/Engine/Plugins/Enterprise/GLTFExporter/Source/GLTFExporter/Public/Json/GLTFJsonRoot.h
--- a/Engine/Plugins/Enterprise/GLTFExporter/Source/GLTFExporter/Public/Json/GLTFJsonRoot.h
+++ b/Engine/Plugins/Enterprise/GLTFExporter/Source/GLTFExporter/Public/Json/GLTFJsonRoot.h
@@ -34,7 +34,8 @@ struct GLTFEXPORTER_API FGLTFJsonRoot : IGLTFJsonObject
 	TGLTFJsonIndexedObjectArray<FGLTFJsonBufferView> BufferViews;
 	TGLTFJsonIndexedObjectArray<FGLTFJsonCamera>     Cameras;
 	TGLTFJsonIndexedObjectArray<FGLTFJsonMaterial>   Materials;
-	TGLTFJsonIndexedObjectArray<FGLTFJsonMesh>       Meshes;
+	TGLTFJsonIndexedObjectArray<FGLTFJsonMesh>       Meshes; // Important! : FGLTFJsonMeshes are validated in "FGLTFJsonBuilder::ValidateAndFixGLTFJson" and any that's found invalid (has no value) it will be removed from the list and deleted.
+															 //					Any references to said deleted item need to be removed as well (for example Nodes.Mesh (1 line below))
 	TGLTFJsonIndexedObjectArray<FGLTFJsonNode>       Nodes;
 	TGLTFJsonIndexedObjectArray<FGLTFJsonImage>      Images;
 	TGLTFJsonIndexedObjectArray<FGLTFJsonSampler>    Samplers;
