[glTF Exporter] Exporting IES source files.

#rnx

diff --git a/Engine/Plugins/Enterprise/GLTFExporter/Source/GLTFExporter/GLTFExporter.Build.cs b/Engine/Plugins/Enterprise/GLTFExporter/Source/GLTFExporter/GLTFExporter.Build.cs
--- a/Engine/Plugins/Enterprise/GLTFExporter/Source/GLTFExporter/GLTFExporter.Build.cs
+++ b/Engine/Plugins/Enterprise/GLTFExporter/Source/GLTFExporter/GLTFExporter.Build.cs
@@ -48,6 +48,9 @@ public class GLTFExporter : ModuleRules
 				"EngineSettings",
 				"Analytics",
 				"Landscape",
+				"InterchangeEngine",
+				"InterchangeCore",
+				"InterchangeFactoryNodes"
 			}
 		);
 
diff --git a/Engine/Plugins/Enterprise/GLTFExporter/Source/GLTFExporter/Private/Builders/GLTFConvertBuilder.cpp b/Engine/Plugins/Enterprise/GLTFExporter/Source/GLTFExporter/Private/Builders/GLTFConvertBuilder.cpp
--- a/Engine/Plugins/Enterprise/GLTFExporter/Source/GLTFExporter/Private/Builders/GLTFConvertBuilder.cpp
+++ b/Engine/Plugins/Enterprise/GLTFExporter/Source/GLTFExporter/Private/Builders/GLTFConvertBuilder.cpp
@@ -577,6 +577,30 @@ FGLTFJsonLight* FGLTFConvertBuilder::AddUniqueLight(const ULightComponent* Light
 	return LightConverter->GetOrAdd(LightComponent);
 }
 
+FGLTFJsonLightIES* FGLTFConvertBuilder::AddUniqueLightIES(const ULightComponent* LightComponent)
+{
+	if (LightComponent == nullptr)
+	{
+		return nullptr;
+	}
+
+	//light are recorded on AddUniqueLight.
+
+	return LightIESConverter->GetOrAdd(LightComponent);	
+}
+
+FGLTFJsonLightIESInstance* FGLTFConvertBuilder::AddUniqueLightIESInstance(const ULightComponent* LightComponent)
+{
+	if (LightComponent == nullptr)
+	{
+		return nullptr;
+	}
+
+	//light are recorded on AddUniqueLight.
+
+	return LightIESInstanceConverter->GetOrAdd(LightComponent);
+}
+
 FGLTFJsonMaterialVariant* FGLTFConvertBuilder::AddUniqueMaterialVariant(const UVariant* Variant)
 {
 	if (Variant == nullptr)
diff --git a/Engine/Plugins/Enterprise/GLTFExporter/Source/GLTFExporter/Private/Builders/GLTFJsonBuilder.cpp b/Engine/Plugins/Enterprise/GLTFExporter/Source/GLTFExporter/Private/Builders/GLTFJsonBuilder.cpp
--- a/Engine/Plugins/Enterprise/GLTFExporter/Source/GLTFExporter/Private/Builders/GLTFJsonBuilder.cpp
+++ b/Engine/Plugins/Enterprise/GLTFExporter/Source/GLTFExporter/Private/Builders/GLTFJsonBuilder.cpp
@@ -103,6 +103,16 @@ FGLTFJsonLight* FGLTFJsonBuilder::AddLight()
 	return JsonRoot.Lights.Add();
 }
 
+FGLTFJsonLightIES* FGLTFJsonBuilder::AddLightIES()
+{
+	return JsonRoot.LightIESs.Add();
+}
+
+FGLTFJsonLightIESInstance* FGLTFJsonBuilder::AddLightIESInstance()
+{
+	return JsonRoot.LightIESInstances.Add();
+}
+
 FGLTFJsonMaterialVariant* FGLTFJsonBuilder::AddMaterialVariant()
 {
 	return JsonRoot.MaterialVariants.Add();
diff --git a/Engine/Plugins/Enterprise/GLTFExporter/Source/GLTFExporter/Private/Converters/GLTFLightConverters.cpp b/Engine/Plugins/Enterprise/GLTFExporter/Source/GLTFExporter/Private/Converters/GLTFLightConverters.cpp
--- a/Engine/Plugins/Enterprise/GLTFExporter/Source/GLTFExporter/Private/Converters/GLTFLightConverters.cpp
+++ b/Engine/Plugins/Enterprise/GLTFExporter/Source/GLTFExporter/Private/Converters/GLTFLightConverters.cpp
@@ -8,6 +8,11 @@
 #include "Components/PointLightComponent.h"
 #include "Components/SpotLightComponent.h"
 
+#if WITH_EDITORONLY_DATA
+#include "InterchangeAssetImportData.h"
+#include "InterchangeTextureLightProfileFactoryNode.h"
+#endif
+
 FGLTFJsonLight* FGLTFLightConverter::Convert(const ULightComponent* LightComponent)
 {
 	const EGLTFJsonLightType Type = FGLTFCoreUtilities::ConvertLightType(LightComponent->GetLightType());
@@ -21,13 +26,12 @@ FGLTFJsonLight* FGLTFLightConverter::Convert(const ULightComponent* LightCompone
 	Light->Name = FGLTFNameUtilities::GetName(LightComponent);
 	Light->Type = Type;
 
-	const float ConversionScale = Type == EGLTFJsonLightType::Directional ? 1.0f : (0.01f * 0.01f); // Conversion from cm2 to m2
-	const FLinearColor ColorBrightness = LightComponent->GetColoredLightBrightness() * ConversionScale;
-	const float Brightness = FMath::Max(ColorBrightness.GetMax(), 1.0f);
-	const FLinearColor Color = ColorBrightness / Brightness;
-
-	Light->Intensity = Brightness;
-	Light->Color = FGLTFCoreUtilities::ConvertColor3(Color);
+	if (!LightComponent->IESTexture || !LightComponent->bUseIESBrightness)
+	{
+		Light->Intensity = LightComponent->Intensity;
+	}
+	
+	Light->Color = FGLTFCoreUtilities::ConvertColor3(LightComponent->LightColor);
 
 	if (const UPointLightComponent* PointLightComponent = Cast<UPointLightComponent>(LightComponent))
 	{
@@ -47,3 +51,77 @@ FGLTFJsonLight* FGLTFLightConverter::Convert(const ULightComponent* LightCompone
 
 	return Light;
 }
+
+FGLTFJsonLightIES* FGLTFLightIESConverter::Convert(const ULightComponent* LightComponent)
+{
+#if WITH_EDITORONLY_DATA
+	if (!LightComponent->IESTexture)
+	{
+		return nullptr;
+	}
+
+
+	FGLTFJsonLightIES* LightIES = Builder.AddLightIES();
+
+	UTextureLightProfile* TextureLightProfile = LightComponent->IESTexture.Get();
+
+	LightIES->Name = TextureLightProfile->GetName();
+
+	if (TextureLightProfile->AssetImportData)
+	{
+		UInterchangeAssetImportData* InterchangeAssetImportData = Cast<UInterchangeAssetImportData>(TextureLightProfile->AssetImportData.Get());
+		if (UInterchangeTextureLightProfileFactoryNode* LightProfileFactoryNode = Cast<UInterchangeTextureLightProfileFactoryNode>(InterchangeAssetImportData->GetNodeContainer()->GetFactoryNode(InterchangeAssetImportData->NodeUniqueID)))
+		{
+			const TSharedPtr<FGLTFMemoryArchive> IESFileContent = MakeShared<FGLTFMemoryArchive>();
+			LightProfileFactoryNode->GetIESSourceFileContents(*IESFileContent);
+
+			if (Builder.bIsGLB)
+			{
+				LightIES->BufferView = Builder.AddBufferView(IESFileContent->GetData(), IESFileContent->Num());
+			}
+			else
+			{
+				LightIES->URI = Builder.AddExternalFile(LightIES->Name + TEXT(".IES"), IESFileContent);
+			}
+		}
+	}
+
+	return LightIES;
+#else
+	Builder.LogWarning(FString::Printf(
+		TEXT("[%s] IES Light Export is not supported at Runtime."),
+		*LightComponent->GetName()));
+
+	return nullptr;
+#endif
+}
+
+FGLTFJsonLightIESInstance* FGLTFLightIESInstanceConverter::Convert(const ULightComponent* LightComponent)
+{
+	if (!LightComponent->IESTexture)
+	{
+		return nullptr;
+	}
+
+	//IES Light:
+	FGLTFJsonLightIES* IESLight = Builder.AddUniqueLightIES(LightComponent);
+
+	if (!IESLight)
+	{
+		//If the IES Light is null, then null the IES Instance as well.
+		return nullptr;
+	}
+
+	FGLTFJsonLightIESInstance* LightIESInstance = Builder.AddLightIESInstance();
+
+	//Multiplier
+	if (LightComponent->bUseIESBrightness)
+	{
+		LightIESInstance->Multiplier = LightComponent->IESBrightnessScale;
+	}
+
+	//Set the IES Light:
+	LightIESInstance->LightIES = IESLight;
+
+	return LightIESInstance;
+}
\ No newline at end of file
diff --git a/Engine/Plugins/Enterprise/GLTFExporter/Source/GLTFExporter/Private/Converters/GLTFNodeConverters.cpp b/Engine/Plugins/Enterprise/GLTFExporter/Source/GLTFExporter/Private/Converters/GLTFNodeConverters.cpp
--- a/Engine/Plugins/Enterprise/GLTFExporter/Source/GLTFExporter/Private/Converters/GLTFNodeConverters.cpp
+++ b/Engine/Plugins/Enterprise/GLTFExporter/Source/GLTFExporter/Private/Converters/GLTFNodeConverters.cpp
@@ -282,9 +282,16 @@ void FGLTFComponentConverter::ConvertComponentSpecialization(const USceneCompone
 		{
 			// TODO: conversion of light direction should be done in separate converter
 			FGLTFJsonNode* LightNode = Builder.AddNode();
-			LightNode->Name = FGLTFNameUtilities::GetName(LightComponent);
+			LightNode->Name = TEXT("LightNode_") + FGLTFNameUtilities::GetName(LightComponent);
 			LightNode->Rotation = FGLTFCoreUtilities::GetLocalLightRotation();
+
 			LightNode->Light = Builder.AddUniqueLight(LightComponent);
+
+			if (LightComponent->IESTexture)
+			{
+				LightNode->LightIESInstance = Builder.AddUniqueLightIESInstance(LightComponent);
+			}
+
 			Node->Children.Add(LightNode);
 		}
 	}
diff --git a/Engine/Plugins/Enterprise/GLTFExporter/Source/GLTFExporter/Private/Json/GLTFJsonLight.cpp b/Engine/Plugins/Enterprise/GLTFExporter/Source/GLTFExporter/Private/Json/GLTFJsonLight.cpp
--- a/Engine/Plugins/Enterprise/GLTFExporter/Source/GLTFExporter/Private/Json/GLTFJsonLight.cpp
+++ b/Engine/Plugins/Enterprise/GLTFExporter/Source/GLTFExporter/Private/Json/GLTFJsonLight.cpp
@@ -1,6 +1,7 @@
 // Copyright Epic Games, Inc. All Rights Reserved.
 
 #include "Json/GLTFJsonLight.h"
+#include "Json/GLTFJsonBufferView.h"
 
 void FGLTFJsonSpotLight::WriteObject(IGLTFJsonWriter& Writer) const
 {
@@ -47,3 +48,36 @@ void FGLTFJsonLight::WriteObject(IGLTFJsonWriter& Writer) const
 		}
 	}
 }
+
+void FGLTFJsonLightIES::WriteObject(IGLTFJsonWriter& Writer) const
+{
+	if (!Name.IsEmpty())
+	{
+		Writer.Write(TEXT("name"), Name);
+	}
+
+	if (!URI.IsEmpty())
+	{
+		Writer.Write(TEXT("uri"), URI);
+	}
+
+	if (BufferView && !MimeType.IsEmpty())
+	{
+		Writer.Write(TEXT("mimeType"), MimeType);
+	}
+
+	if (BufferView != nullptr)
+	{
+		Writer.Write(TEXT("bufferView"), BufferView);
+	}
+}
+
+void FGLTFJsonLightIESInstance::WriteObject(IGLTFJsonWriter& Writer) const
+{
+	if (!FMath::IsNearlyEqual(Multiplier, 1, Writer.DefaultTolerance))
+	{
+		Writer.Write(TEXT("multiplier"), Multiplier);
+	}
+
+	Writer.Write(TEXT("light"), LightIES);
+}
\ No newline at end of file
diff --git a/Engine/Plugins/Enterprise/GLTFExporter/Source/GLTFExporter/Private/Json/GLTFJsonNode.cpp b/Engine/Plugins/Enterprise/GLTFExporter/Source/GLTFExporter/Private/Json/GLTFJsonNode.cpp
--- a/Engine/Plugins/Enterprise/GLTFExporter/Source/GLTFExporter/Private/Json/GLTFJsonNode.cpp
+++ b/Engine/Plugins/Enterprise/GLTFExporter/Source/GLTFExporter/Private/Json/GLTFJsonNode.cpp
@@ -30,13 +30,23 @@ void FGLTFJsonNode::WriteObject(IGLTFJsonWriter& Writer) const
 		Writer.Write(TEXT("mesh"), Mesh);
 	}
 
-	if (Light != nullptr)
+	if (Light != nullptr
+		|| (LightIESInstance != nullptr && LightIESInstance->HasValue()))
 	{
 		Writer.StartExtensions();
 
-		Writer.StartExtension(EGLTFJsonExtension::KHR_LightsPunctual);
-		Writer.Write(TEXT("light"), Light);
-		Writer.EndExtension();
+		if (Light != nullptr)
+		{
+			Writer.StartExtension(EGLTFJsonExtension::KHR_LightsPunctual);
+			Writer.Write(TEXT("light"), Light);
+			Writer.EndExtension();
+		}
+		
+		if ((LightIESInstance != nullptr && LightIESInstance->HasValue())) {
+			Writer.StartExtension(EGLTFJsonExtension::EXT_LightsIES);
+			LightIESInstance->WriteObject(Writer);
+			Writer.EndExtension();
+		}
 
 		Writer.EndExtensions();
 	}
diff --git a/Engine/Plugins/Enterprise/GLTFExporter/Source/GLTFExporter/Private/Json/GLTFJsonRoot.cpp b/Engine/Plugins/Enterprise/GLTFExporter/Source/GLTFExporter/Private/Json/GLTFJsonRoot.cpp
--- a/Engine/Plugins/Enterprise/GLTFExporter/Source/GLTFExporter/Private/Json/GLTFJsonRoot.cpp
+++ b/Engine/Plugins/Enterprise/GLTFExporter/Source/GLTFExporter/Private/Json/GLTFJsonRoot.cpp
@@ -43,6 +43,13 @@ void FGLTFJsonRoot::WriteObject(IGLTFJsonWriter& Writer) const
 			Writer.EndExtension();
 		}
 
+		if (LightIESs.Num() > 0)
+		{
+			Writer.StartExtension(EGLTFJsonExtension::EXT_LightsIES);
+			Writer.Write(TEXT("lights"), LightIESs);
+			Writer.EndExtension();
+		}
+
 		Writer.EndExtensions();
 	}
 
diff --git a/Engine/Plugins/Enterprise/GLTFExporter/Source/GLTFExporter/Private/Utilities/GLTFJsonUtilities.cpp b/Engine/Plugins/Enterprise/GLTFExporter/Source/GLTFExporter/Private/Utilities/GLTFJsonUtilities.cpp
--- a/Engine/Plugins/Enterprise/GLTFExporter/Source/GLTFExporter/Private/Utilities/GLTFJsonUtilities.cpp
+++ b/Engine/Plugins/Enterprise/GLTFExporter/Source/GLTFExporter/Private/Utilities/GLTFJsonUtilities.cpp
@@ -7,6 +7,7 @@ const TCHAR* FGLTFJsonUtilities::GetValue(EGLTFJsonExtension Enum)
 	switch (Enum)
 	{
 		case EGLTFJsonExtension::KHR_LightsPunctual:                  return TEXT("KHR_lights_punctual");
+		case EGLTFJsonExtension::EXT_LightsIES:                       return TEXT("EXT_lights_ies");
 		case EGLTFJsonExtension::KHR_MaterialsClearCoat:              return TEXT("KHR_materials_clearcoat");
 		case EGLTFJsonExtension::KHR_MaterialsEmissiveStrength:       return TEXT("KHR_materials_emissive_strength");
 		case EGLTFJsonExtension::KHR_MaterialsUnlit:                  return TEXT("KHR_materials_unlit");
diff --git a/Engine/Plugins/Enterprise/GLTFExporter/Source/GLTFExporter/Public/Builders/GLTFConvertBuilder.h b/Engine/Plugins/Enterprise/GLTFExporter/Source/GLTFExporter/Public/Builders/GLTFConvertBuilder.h
--- a/Engine/Plugins/Enterprise/GLTFExporter/Source/GLTFExporter/Public/Builders/GLTFConvertBuilder.h
+++ b/Engine/Plugins/Enterprise/GLTFExporter/Source/GLTFExporter/Public/Builders/GLTFConvertBuilder.h
@@ -97,6 +97,8 @@ public:
 
 	FGLTFJsonCamera* AddUniqueCamera(const UCameraComponent* CameraComponent);
 	FGLTFJsonLight* AddUniqueLight(const ULightComponent* LightComponent);
+	FGLTFJsonLightIES* AddUniqueLightIES(const ULightComponent* LightComponent);
+	FGLTFJsonLightIESInstance* AddUniqueLightIESInstance(const ULightComponent* LightComponent);
 
 	FGLTFJsonMaterialVariant* AddUniqueMaterialVariant(const UVariant* Variant);
 
@@ -156,6 +158,8 @@ public:
 
 	TUniquePtr<IGLTFCameraConverter> CameraConverter = MakeUnique<FGLTFCameraConverter>(*this);
 	TUniquePtr<IGLTFLightConverter> LightConverter = MakeUnique<FGLTFLightConverter>(*this);
+	TUniquePtr<IGLTFLightIESConverter> LightIESConverter = MakeUnique<FGLTFLightIESConverter>(*this);
+	TUniquePtr<FGLTFLightIESInstanceConverter> LightIESInstanceConverter = MakeUnique<FGLTFLightIESInstanceConverter>(*this);
 	TUniquePtr<IGLTFMaterialVariantConverter> MaterialVariantConverter = MakeUnique<FGLTFMaterialVariantConverter>(*this);
 
 private:
diff --git a/Engine/Plugins/Enterprise/GLTFExporter/Source/GLTFExporter/Public/Builders/GLTFJsonBuilder.h b/Engine/Plugins/Enterprise/GLTFExporter/Source/GLTFExporter/Public/Builders/GLTFJsonBuilder.h
--- a/Engine/Plugins/Enterprise/GLTFExporter/Source/GLTFExporter/Public/Builders/GLTFJsonBuilder.h
+++ b/Engine/Plugins/Enterprise/GLTFExporter/Source/GLTFExporter/Public/Builders/GLTFJsonBuilder.h
@@ -29,6 +29,8 @@ public:
 	FGLTFJsonSkin* AddSkin();
 	FGLTFJsonTexture* AddTexture();
 	FGLTFJsonLight* AddLight();
+	FGLTFJsonLightIES* AddLightIES();
+	FGLTFJsonLightIESInstance* AddLightIESInstance();
 	FGLTFJsonMaterialVariant* AddMaterialVariant();
 
 	const FGLTFJsonRoot& GetRoot() const;
diff --git a/Engine/Plugins/Enterprise/GLTFExporter/Source/GLTFExporter/Public/Converters/GLTFLightConverters.h b/Engine/Plugins/Enterprise/GLTFExporter/Source/GLTFExporter/Public/Converters/GLTFLightConverters.h
--- a/Engine/Plugins/Enterprise/GLTFExporter/Source/GLTFExporter/Public/Converters/GLTFLightConverters.h
+++ b/Engine/Plugins/Enterprise/GLTFExporter/Source/GLTFExporter/Public/Converters/GLTFLightConverters.h
@@ -9,6 +9,9 @@
 class ULightComponent;
 
 typedef TGLTFConverter<FGLTFJsonLight*, const ULightComponent*> IGLTFLightConverter;
+typedef TGLTFConverter<FGLTFJsonLightIES*, const ULightComponent*> IGLTFLightIESConverter;
+typedef TGLTFConverter<FGLTFJsonLightIESInstance*, const ULightComponent*> IGLTFLightIESInstanceConverter;
+
 
 class GLTFEXPORTER_API FGLTFLightConverter : public FGLTFBuilderContext, public IGLTFLightConverter
 {
@@ -20,3 +23,27 @@ protected:
 
 	virtual FGLTFJsonLight* Convert(const ULightComponent* LightComponent) override;
 };
+
+//IES export is not supported runtime, as we don't have access to the AssetImportData (which stores the SourceFileContent for the IES file).
+class GLTFEXPORTER_API FGLTFLightIESConverter : public FGLTFBuilderContext, public IGLTFLightIESConverter
+{
+public:
+
+	using FGLTFBuilderContext::FGLTFBuilderContext;
+
+protected:
+
+	virtual FGLTFJsonLightIES* Convert(const ULightComponent* LightComponent) override;
+};
+
+//IES export is not supported runtime, as we don't have access to the AssetImportData (which stores the SourceFileContent for the IES file).
+class GLTFEXPORTER_API FGLTFLightIESInstanceConverter : public FGLTFBuilderContext, public IGLTFLightIESInstanceConverter
+{
+public:
+
+	using FGLTFBuilderContext::FGLTFBuilderContext;
+
+protected:
+
+	virtual FGLTFJsonLightIESInstance* Convert(const ULightComponent* LightComponent) override;
+};
diff --git a/Engine/Plugins/Enterprise/GLTFExporter/Source/GLTFExporter/Public/Json/GLTFJsonCore.h b/Engine/Plugins/Enterprise/GLTFExporter/Source/GLTFExporter/Public/Json/GLTFJsonCore.h
--- a/Engine/Plugins/Enterprise/GLTFExporter/Source/GLTFExporter/Public/Json/GLTFJsonCore.h
+++ b/Engine/Plugins/Enterprise/GLTFExporter/Source/GLTFExporter/Public/Json/GLTFJsonCore.h
@@ -32,6 +32,8 @@ struct FGLTFJsonSkin;
 struct FGLTFJsonTexture;
 struct FGLTFJsonBackdrop;
 struct FGLTFJsonLight;
+struct FGLTFJsonLightIES;
+struct FGLTFJsonLightIESInstance;
 struct FGLTFJsonLightMap;
 struct FGLTFJsonSkySphere;
 struct FGLTFJsonEpicLevelVariantSets;
diff --git a/Engine/Plugins/Enterprise/GLTFExporter/Source/GLTFExporter/Public/Json/GLTFJsonEnums.h b/Engine/Plugins/Enterprise/GLTFExporter/Source/GLTFExporter/Public/Json/GLTFJsonEnums.h
--- a/Engine/Plugins/Enterprise/GLTFExporter/Source/GLTFExporter/Public/Json/GLTFJsonEnums.h
+++ b/Engine/Plugins/Enterprise/GLTFExporter/Source/GLTFExporter/Public/Json/GLTFJsonEnums.h
@@ -5,6 +5,7 @@
 enum class EGLTFJsonExtension
 {
 	KHR_LightsPunctual,
+	EXT_LightsIES,
 	KHR_MaterialsClearCoat,
 	KHR_MaterialsEmissiveStrength,
 	KHR_MaterialsUnlit,
diff --git a/Engine/Plugins/Enterprise/GLTFExporter/Source/GLTFExporter/Public/Json/GLTFJsonLight.h b/Engine/Plugins/Enterprise/GLTFExporter/Source/GLTFExporter/Public/Json/GLTFJsonLight.h
--- a/Engine/Plugins/Enterprise/GLTFExporter/Source/GLTFExporter/Public/Json/GLTFJsonLight.h
+++ b/Engine/Plugins/Enterprise/GLTFExporter/Source/GLTFExporter/Public/Json/GLTFJsonLight.h
@@ -46,3 +46,50 @@ protected:
 	{
 	}
 };
+
+
+struct GLTFEXPORTER_API FGLTFJsonLightIES : IGLTFJsonIndexedObject
+{
+	FString Name;
+
+	FString URI;
+
+	FString MimeType;
+	FGLTFJsonBufferView* BufferView;
+
+	virtual void WriteObject(IGLTFJsonWriter& Writer) const override;
+
+protected:
+
+	friend TGLTFJsonIndexedObjectArray<FGLTFJsonLightIES, void>;
+
+	FGLTFJsonLightIES(int32 Index)
+		: IGLTFJsonIndexedObject(Index)
+		, MimeType(TEXT("application/x-ies-lm-63"))
+		, BufferView(nullptr)
+	{
+	}
+};
+
+struct GLTFEXPORTER_API FGLTFJsonLightIESInstance : IGLTFJsonIndexedObject
+{
+	//FGLTFJsonColor3 Color; //Note: While IES extension can store the Color, the Export will save out the Color within the KHR_Lights_Punctual extension.
+	float Multiplier;
+
+	virtual void WriteObject(IGLTFJsonWriter& Writer) const override;
+
+	FGLTFJsonLightIES* LightIES;
+
+	bool HasValue() { return LightIES && (LightIES->URI.Len() || LightIES->BufferView); }
+
+protected:
+
+	friend TGLTFJsonIndexedObjectArray<FGLTFJsonLightIESInstance, void>;
+
+	FGLTFJsonLightIESInstance(int32 Index)
+		: IGLTFJsonIndexedObject(Index)
+		, Multiplier(1)
+		, LightIES(nullptr)
+	{
+	}
+};
\ No newline at end of file
diff --git a/Engine/Plugins/Enterprise/GLTFExporter/Source/GLTFExporter/Public/Json/GLTFJsonNode.h b/Engine/Plugins/Enterprise/GLTFExporter/Source/GLTFExporter/Public/Json/GLTFJsonNode.h
--- a/Engine/Plugins/Enterprise/GLTFExporter/Source/GLTFExporter/Public/Json/GLTFJsonNode.h
+++ b/Engine/Plugins/Enterprise/GLTFExporter/Source/GLTFExporter/Public/Json/GLTFJsonNode.h
@@ -13,6 +13,7 @@ struct GLTFEXPORTER_API FGLTFJsonNode : FGLTFJsonTransform, IGLTFJsonIndexedObje
 	FGLTFJsonSkin*   Skin;
 	FGLTFJsonMesh*   Mesh;
 	FGLTFJsonLight*  Light;
+	FGLTFJsonLightIESInstance*  LightIESInstance;
 
 	TArray<FGLTFJsonNode*> Children;
 
@@ -28,6 +29,7 @@ protected:
 		, Skin(nullptr)
 		, Mesh(nullptr)
 		, Light(nullptr)
+		, LightIESInstance(nullptr)
 	{
 	}
 };
diff --git a/Engine/Plugins/Enterprise/GLTFExporter/Source/GLTFExporter/Public/Json/GLTFJsonRoot.h b/Engine/Plugins/Enterprise/GLTFExporter/Source/GLTFExporter/Public/Json/GLTFJsonRoot.h
--- a/Engine/Plugins/Enterprise/GLTFExporter/Source/GLTFExporter/Public/Json/GLTFJsonRoot.h
+++ b/Engine/Plugins/Enterprise/GLTFExporter/Source/GLTFExporter/Public/Json/GLTFJsonRoot.h
@@ -43,6 +43,8 @@ struct GLTFEXPORTER_API FGLTFJsonRoot : IGLTFJsonObject
 	TGLTFJsonIndexedObjectArray<FGLTFJsonSkin>       Skins;
 	TGLTFJsonIndexedObjectArray<FGLTFJsonTexture>    Textures;
 	TGLTFJsonIndexedObjectArray<FGLTFJsonLight>      Lights;
+	TGLTFJsonIndexedObjectArray<FGLTFJsonLightIES>   LightIESs;
+	TGLTFJsonIndexedObjectArray<FGLTFJsonLightIESInstance>  LightIESInstances;
 	TGLTFJsonIndexedObjectArray<FGLTFJsonMaterialVariant>   MaterialVariants;
 
 	FGLTFJsonRoot()
diff --git a/Engine/Plugins/Enterprise/GLTFExporter/Source/GLTFExporter/Public/Options/GLTFExportOptions.h b/Engine/Plugins/Enterprise/GLTFExporter/Source/GLTFExporter/Public/Options/GLTFExportOptions.h
--- a/Engine/Plugins/Enterprise/GLTFExporter/Source/GLTFExporter/Public/Options/GLTFExportOptions.h
+++ b/Engine/Plugins/Enterprise/GLTFExporter/Source/GLTFExporter/Public/Options/GLTFExportOptions.h
@@ -153,7 +153,9 @@ class GLTFEXPORTER_API UGLTFExportOptions : public UObject
 	UPROPERTY(EditAnywhere, BlueprintReadWrite, Config, Category = Scene)
 	bool bExportHiddenInGame;
 
-	/** If enabled, export directional, point, and spot light components. Uses extension KHR_lights_punctual. */
+	/** If enabled, export directional, point, and spot light components. Uses extension KHR_lights_punctual.
+	* Also exports IES light profiles using the EXT_lights_ies extension.
+	*/
 	UPROPERTY(EditAnywhere, BlueprintReadWrite, Config, Category = Scene)
 	bool bExportLights;
 
