[glTF Exporter] Refactoring SkeletalMesh export to use SourceModel instead of ImportedModel.

#rnx

diff --git a/Engine/Plugins/Enterprise/GLTFExporter/Source/GLTFExporter/GLTFExporter.Build.cs b/Engine/Plugins/Enterprise/GLTFExporter/Source/GLTFExporter/GLTFExporter.Build.cs
--- a/Engine/Plugins/Enterprise/GLTFExporter/Source/GLTFExporter/GLTFExporter.Build.cs
+++ b/Engine/Plugins/Enterprise/GLTFExporter/Source/GLTFExporter/GLTFExporter.Build.cs
@@ -50,7 +50,8 @@ public class GLTFExporter : ModuleRules
 				"Landscape",
 				"InterchangeEngine",
 				"InterchangeCore",
-				"InterchangeFactoryNodes"
+				"InterchangeFactoryNodes",
+				"SkeletalMeshDescription"
 			}
 		);
 
diff --git a/Engine/Plugins/Enterprise/GLTFExporter/Source/GLTFExporter/Private/Tasks/GLTFDelayedMeshTasks.cpp b/Engine/Plugins/Enterprise/GLTFExporter/Source/GLTFExporter/Private/Tasks/GLTFDelayedMeshTasks.cpp
--- a/Engine/Plugins/Enterprise/GLTFExporter/Source/GLTFExporter/Private/Tasks/GLTFDelayedMeshTasks.cpp
+++ b/Engine/Plugins/Enterprise/GLTFExporter/Source/GLTFExporter/Private/Tasks/GLTFDelayedMeshTasks.cpp
@@ -13,6 +13,7 @@
 #include "StaticMeshResources.h"
 #include "Rendering/SkeletalMeshRenderData.h"
 #include "Rendering/SkeletalMeshModel.h"
+#include "SkeletalMeshAttributes.h"
 
 #include "Converters/GLTFMaterialUtilities.h"
 #include "Rendering/PositionVertexBuffer.h"
@@ -23,6 +24,10 @@
 #include "Converters/GLTFMeshAttributesArray.h"
 #include "Utilities/GLTFLandscapeComponentDataInterface.h"
 
+#if WITH_EDITORONLY_DATA
+#include "Utilities/MeshParsingUtilities.h"
+#endif
+
 namespace
 {
 	template <typename VectorType>
@@ -181,220 +186,52 @@ void FGLTFDelayedStaticAndSplineMeshTask::Process()
 }
 
 #if WITH_EDITORONLY_DATA
+
 void FGLTFDelayedStaticAndSplineMeshTask::ProcessMeshDescription(const TArray<FStaticMaterial>& MaterialSlots, const FGLTFMeshData* MeshData)
 {
 	FMeshDescription* MeshDescription = StaticMesh->GetMeshDescription(LODIndex);
-	const FPolygonGroupArray& PolygonGroups = MeshDescription->PolygonGroups();
 
-	if (PolygonGroups.Num() != JsonMesh->Primitives.Num()
-		|| PolygonGroups.Num() != MaterialSlots.Num())
-	{
-		return;
-	}
+	using namespace UE::MeshParser;
 
-	FStaticMeshConstAttributes Attributes(*MeshDescription);
-	TVertexAttributesConstRef<FVector3f> VertexPositions = Attributes.GetVertexPositions();
-	TVertexInstanceAttributesConstRef<FVector3f> VertexInstanceNormals = Attributes.GetVertexInstanceNormals();
-	TVertexInstanceAttributesConstRef<FVector3f> VertexInstanceTangents = Attributes.GetVertexInstanceTangents();
-	TVertexInstanceAttributesConstRef<float> VertexInstanceBinormalSigns = Attributes.GetVertexInstanceBinormalSigns();
-	TVertexInstanceAttributesConstRef<FVector2f> VertexInstanceUVs = Attributes.GetVertexInstanceUVs();
-	TVertexInstanceAttributesConstRef<FVector4f> VertexInstanceColors = Attributes.GetVertexInstanceColors();
-	TPolygonGroupAttributesConstRef<FName> PolygonGroupMaterialSlotNames = Attributes.GetPolygonGroupMaterialSlotNames();
-	TEdgeAttributesConstRef<bool> EdgeHardnesses = Attributes.GetEdgeHardnesses();
+	FMeshDescriptionParser MeshDescriptionParser(MeshDescription, MaterialSlots);
 
-	const int32 VertexCount = MeshDescription->Vertices().Num();
-	const int32 VertexInstanceCount = MeshDescription->VertexInstances().Num();
+	constexpr bool bExportVertexSkinWeights_False = false;
+	constexpr int32 SkeletonInfluenceCountPerGroup_4 = 4;
+	FExportConfigs ExportConfigs(bExportVertexSkinWeights_False, Builder.ExportOptions->bExportVertexColors, SplineMeshComponent, SkeletonInfluenceCountPerGroup_4);
 
-	bool bHasVertexColors = Builder.ExportOptions->bExportVertexColors && MeshDescription->VertexInstanceAttributes().HasAttribute(MeshAttribute::VertexInstance::Color);
-	int32 UVCount = VertexInstanceUVs.GetNumChannels();
+	TArray<FMeshPrimitiveDescription> MeshPrimitiveDescriptions;
+	MeshDescriptionParser.Parse(MeshPrimitiveDescriptions, ExportConfigs);
+	
+	if (MeshPrimitiveDescriptions.Num() != JsonMesh->Primitives.Num())
+	{
+		return;
+	}
 
 	int32 PrimitiveIndex = 0;
-	int32 NumberOfPrimitives = MeshDescription->PolygonGroups().Num();
-	for (const FPolygonGroupID& PolygonGroupID : MeshDescription->PolygonGroups().GetElementIDs())
+	for (FMeshPrimitiveDescription& MeshPrimitiveDescription : MeshPrimitiveDescriptions)
 	{
-		const TArrayView<const FTriangleID>& TriangleIDs = MeshDescription->GetPolygonGroupTriangles(PolygonGroupID);
-
-		if (TriangleIDs.Num() == 0)
+		if (MeshPrimitiveDescription.IsEmpty())
 		{
 			PrimitiveIndex++;
-			//Do not export empty primitives.
 			continue;
 		}
 
-		FName MaterialSlotName = PolygonGroupMaterialSlotNames[PolygonGroupID];
 		FGLTFJsonPrimitive& JsonPrimitive = JsonMesh->Primitives[PrimitiveIndex];
 
-		int32 MaterialIndex = INDEX_NONE;
-		{
-			for (size_t MatIndex = 0; MatIndex < MaterialSlots.Num(); MatIndex++)
-			{
-				if (MaterialSlots[MatIndex].ImportedMaterialSlotName == MaterialSlotName)
-				{
-					MaterialIndex = MatIndex;
-					break;
-				}
-			}
-
-			if (MaterialIndex == INDEX_NONE)
-			{
-				MaterialIndex = PolygonGroupID.GetValue();
-			}
-			if (!MaterialSlots.IsValidIndex(MaterialIndex))
-			{
-				MaterialIndex = 0;
-			}
-		}
-
-		TArray<int32> OriginalIndices;
-		TArray<FVector3f> OriginalPositions;
-		TArray<FColor> OriginalVertexColors;
-		TArray<FVector3f> OriginalNormals;
-		TArray<FVector4f> OriginalTangents;
-		TArray<TArray<FVector2f>> OriginalUVs;
-
-		OriginalIndices.Reserve(VertexInstanceCount);
-		OriginalPositions.SetNumZeroed(VertexInstanceCount);
-		OriginalVertexColors.SetNumZeroed(VertexInstanceCount);
-		OriginalNormals.SetNumZeroed(VertexInstanceCount);
-		OriginalTangents.SetNumZeroed(VertexInstanceCount);
-		for (size_t UVIndex = 0; UVIndex < UVCount; UVIndex++)
-		{
-			OriginalUVs.Add(TArray<FVector2f>());
-			OriginalUVs[UVIndex].SetNumZeroed(VertexInstanceCount);
-		}
-
-		for (const FTriangleID& TriangleID : TriangleIDs)
-		{
-			TArrayView<const FVertexInstanceID> TriangleVertexInstanceIDs = MeshDescription->GetTriangleVertexInstances(TriangleID);
-			for (const FVertexInstanceID& VertexInstanceID : TriangleVertexInstanceIDs)
-			{
-				OriginalIndices.Add(VertexInstanceID);
-
-				FVector3f Position = VertexPositions[MeshDescription->GetVertexInstanceVertex(VertexInstanceID).GetValue()];
-				OriginalPositions[VertexInstanceID] = Position;
-
-				OriginalNormals[VertexInstanceID] = VertexInstanceNormals[VertexInstanceID];
-				OriginalTangents[VertexInstanceID] = FVector4f(VertexInstanceTangents[VertexInstanceID], VertexInstanceBinormalSigns[VertexInstanceID]);
-
-				for (size_t UVIndex = 0; UVIndex < UVCount; UVIndex++)
-				{
-					OriginalUVs[UVIndex][VertexInstanceID] = VertexInstanceUVs.Get(VertexInstanceID, UVIndex);
-				}
-			}
-		}
-
-		if (bHasVertexColors)
-		{
-			for (const FTriangleID& TriangleID : TriangleIDs)
-			{
-				TArrayView<const FVertexInstanceID> TriangleVertexInstanceIDs = MeshDescription->GetTriangleVertexInstances(TriangleID);
-				for (const FVertexInstanceID& VertexInstanceID : TriangleVertexInstanceIDs)
-				{
-					const FVector4f& SourceVertexColor = VertexInstanceColors[VertexInstanceID];
-					OriginalVertexColors[VertexInstanceID] = FLinearColor(SourceVertexColor).ToFColor(true);
-				}
-			}
-		}
-
-		FGLTFIndexArray Indices;
-		FGLTFPositionArray PositionBuffer;
-		FGLTFColorArray VertexColorBuffer;
-		FGLTFNormalArray Normals;
-		FGLTFTangentArray Tangents;
-		TArray<FGLTFUVArray> UVs;
-
-		//Remap Containers to contain only used data sets (per primitive)
-		// + Fill glTF Containers.
-		{
-			TArray<int32> SortedIndices = OriginalIndices;
-			TMap<int32, int32> IndexRemapper;
-			{
-				SortedIndices.Sort();
-
-				TSet<int32> SortedUniqueIndices;
-				SortedUniqueIndices.Append(SortedIndices);
-				SortedIndices.Empty();
-
-				SortedIndices = SortedUniqueIndices.Array();
-
-				for (size_t NewIndex = 0; NewIndex < SortedIndices.Num(); NewIndex++)
-				{
-					IndexRemapper.Add(SortedIndices[NewIndex], NewIndex);
-				}
-			}
-
-			int32 PrimtiveVertexCount = IndexRemapper.Num();
-
-			Indices.Reserve(PrimtiveVertexCount);
-			PositionBuffer.SetNumZeroed(PrimtiveVertexCount);
-			VertexColorBuffer.SetNumZeroed(PrimtiveVertexCount);
-			Normals.SetNumZeroed(PrimtiveVertexCount);
-			Tangents.SetNumZeroed(PrimtiveVertexCount);
-
-			UVs.AddDefaulted(UVCount);
-			for (size_t UVIndex = 0; UVIndex < UVCount; UVIndex++)
-			{
-				UVs[UVIndex].SetNumZeroed(PrimtiveVertexCount);
-			}
-
-			for (const int32& OriginalIndex : OriginalIndices)
-			{
-				int32 NewIndex = IndexRemapper[OriginalIndex];
-				Indices.Add(NewIndex);
-
-				FVector3f Position = OriginalPositions[OriginalIndex];
-				if (SplineMeshComponent)
-				{
-					//SplineMeshComponent provided
-					//Fix the Positions for Splines:
-					const FTransform3f SliceTransform = FTransform3f(SplineMeshComponent->CalcSliceTransform(USplineMeshComponent::GetAxisValueRef(Position, SplineMeshComponent->ForwardAxis)));
-					USplineMeshComponent::GetAxisValueRef(Position, SplineMeshComponent->ForwardAxis) = 0;
-					Position = SliceTransform.TransformPosition(Position);
-				}
-
-				PositionBuffer[NewIndex] = Position;
-
-				if (bHasVertexColors)
-				{
-					VertexColorBuffer[NewIndex] = OriginalVertexColors[OriginalIndex];
-				}
-
-				Normals[NewIndex] = OriginalNormals[OriginalIndex];
-				Tangents[NewIndex] = OriginalTangents[OriginalIndex];
-
-				for (size_t UVIndex = 0; UVIndex < UVCount; UVIndex++)
-				{
-					UVs[UVIndex][NewIndex] = OriginalUVs[UVIndex][OriginalIndex];
-				}
-			}
-		}
-
-		if (Tangents.Num() > 0 && DoesBufferHasZeroVector(Tangents))
-		{
-			//Do not Export Tangents list that is zeroed out.
-			Tangents.Reset();
-		}
-
-		if (Normals.Num() > 0 && DoesBufferHasZeroVector(Normals))
-		{
-			//Do not Export Normals list that is zeroed out.
-			Normals.Reset();
-		}
-
 		//Set glTF Primitive:
-		JsonPrimitive.Indices = Builder.AddUniqueIndexAccessor(Indices, StaticMesh->GetName() + (NumberOfPrimitives > 1 ? (TEXT("_") + FString::FromInt(PrimitiveIndex)) : TEXT("")));
-		JsonPrimitive.Attributes.Position = Builder.AddUniquePositionAccessor(PositionBuffer);
-		if (bHasVertexColors) JsonPrimitive.Attributes.Color0 = Builder.AddUniqueColorAccessor(VertexColorBuffer);
-		JsonPrimitive.Attributes.Normal = Builder.AddUniqueNormalAccessor(Normals);
-		JsonPrimitive.Attributes.Tangent = Builder.AddUniqueTangentAccessor(Tangents);
-		JsonPrimitive.Attributes.TexCoords.AddUninitialized(UVCount);
-		for (size_t UVIndex = 0; UVIndex < UVCount; UVIndex++)
+		JsonPrimitive.Indices = Builder.AddUniqueIndexAccessor(MeshPrimitiveDescription.Indices, StaticMesh->GetName() + (MeshDescriptionParser.MeshDetails.NumberOfPrimitives > 1 ? (TEXT("_") + FString::FromInt(PrimitiveIndex)) : TEXT("")));
+		JsonPrimitive.Attributes.Position = Builder.AddUniquePositionAccessor(MeshPrimitiveDescription.Positions);
+		if (!MeshPrimitiveDescription.VertexColors.IsEmpty()) JsonPrimitive.Attributes.Color0 = Builder.AddUniqueColorAccessor(MeshPrimitiveDescription.VertexColors);
+		JsonPrimitive.Attributes.Normal = Builder.AddUniqueNormalAccessor(MeshPrimitiveDescription.Normals);
+		JsonPrimitive.Attributes.Tangent = Builder.AddUniqueTangentAccessor(MeshPrimitiveDescription.Tangents);
+		JsonPrimitive.Attributes.TexCoords.AddUninitialized(MeshDescriptionParser.MeshDetails.UVCount);
+		for (size_t UVIndex = 0; UVIndex < MeshDescriptionParser.MeshDetails.UVCount; UVIndex++)
 		{
-			JsonPrimitive.Attributes.TexCoords[UVIndex] = Builder.AddUniqueUVAccessor(UVs[UVIndex]);
+			JsonPrimitive.Attributes.TexCoords[UVIndex] = Builder.AddUniqueUVAccessor(FGLTFUVArray(MeshPrimitiveDescription.UVs[UVIndex]));
 		}
 		//
-		const UMaterialInterface* Material = Materials.IsValidIndex(MaterialIndex) ? Materials[MaterialIndex] : MaterialSlots[MaterialIndex].MaterialInterface;
-		JsonPrimitive.Material = Builder.AddUniqueMaterial(Material, MeshData, { MaterialIndex });
+		const UMaterialInterface* Material = Materials.IsValidIndex(MeshPrimitiveDescription.MaterialIndex) ? Materials[MeshPrimitiveDescription.MaterialIndex] : MaterialSlots[MeshPrimitiveDescription.MaterialIndex].MaterialInterface;
+		JsonPrimitive.Material = Builder.AddUniqueMaterial(Material, MeshData, { MeshPrimitiveDescription.MaterialIndex });
 
 		//Validations:
 		if (JsonPrimitive.Attributes.Position == nullptr)
@@ -568,206 +405,71 @@ void FGLTFDelayedSkeletalMeshTask::Process()
 #if WITH_EDITORONLY_DATA
 void FGLTFDelayedSkeletalMeshTask::ProcessSourceModel(const TArray<FSkeletalMaterial>& MaterialSlots, const FGLTFMeshData* MeshData)
 {
-	FSkeletalMeshModel* ImportedSkeletalMeshModel = SkeletalMesh->GetImportedModel();
-	if (!ImportedSkeletalMeshModel->LODModels.IsValidIndex(LODIndex))
+	if (SkeletalMesh->GetNumSourceModels() < LODIndex)
 	{
-		//TODO: Log Error
 		return;
 	}
+	const FSkeletalMeshSourceModel& MeshSourceModel = SkeletalMesh->GetSourceModel(LODIndex);
+	const FMeshDescription* MeshDescription = MeshSourceModel.GetMeshDescription();
 
-	bool bExportVertexColors = Builder.ExportOptions->bExportVertexColors;
+	using namespace UE::MeshParser;
+	
+	FMeshDescriptionParser MeshDescriptionParser(MeshDescription, MaterialSlots);
+	
+	constexpr bool bExportVertexSkinWeights_True = true;
+	constexpr int32 SkeletonInfluenceCountPerGroup_4 = 4;
+	FExportConfigs ExportConfigs(bExportVertexSkinWeights_True, Builder.ExportOptions->bExportVertexColors, nullptr, SkeletonInfluenceCountPerGroup_4);
+	
 
-	FSkeletalMeshLODModel& SourceModel = ImportedSkeletalMeshModel->LODModels[LODIndex];
-	const int32 SectionCount = SourceModel.Sections.Num();
-	const int32 NumTexCoords = SourceModel.NumTexCoords;
+	TArray<FMeshPrimitiveDescription> MeshPrimitiveDescriptions;
+	MeshDescriptionParser.Parse(MeshPrimitiveDescriptions, ExportConfigs);
 
-	if (JsonMesh->Primitives.Num() != SectionCount)
+	if (MeshPrimitiveDescriptions.Num() != JsonMesh->Primitives.Num())
 	{
-		//TODO: Log Error
 		return;
 	}
 
-	TArray<FSoftSkinVertex> SoftSkinVertices;
-	SourceModel.GetNonClothVertices(SoftSkinVertices);
-
-	int32 ClothSectionVertexRemoveOffset = 0;
-	for (int32 SectionIndex = 0; SectionIndex < SectionCount; ++SectionIndex)
+	int32 PrimitiveIndex = 0;
+	for (FMeshPrimitiveDescription& MeshPrimitiveDescription : MeshPrimitiveDescriptions)
 	{
-		FGLTFJsonPrimitive& JsonPrimitive = JsonMesh->Primitives[SectionIndex];
-
-		const FSkelMeshSection& Section = SourceModel.Sections[SectionIndex];
-		//Section.MaterialIndex
-		if (Section.HasClothingData() || !(MaterialSlots.IsValidIndex(Section.MaterialIndex)))
+		if (MeshPrimitiveDescription.IsEmpty())
 		{
-			ClothSectionVertexRemoveOffset += Section.GetNumVertices();
+			PrimitiveIndex++;
 			continue;
 		}
 
-		int32 TriangleCount = Section.NumTriangles;
-
-		TArray<int32> OriginalIndices;
-		OriginalIndices.Reserve(TriangleCount * 3);
-
-		// Copy over the index buffer into the FBX polygons set.
-		for (int32 TriangleIndex = 0; TriangleIndex < TriangleCount; ++TriangleIndex)
-		{
-			for (int32 PointIndex = 0; PointIndex < 3; PointIndex++)
-			{
-				int32 VertexPositionIndex = SourceModel.IndexBuffer[Section.BaseIndex + ((TriangleIndex * 3) + PointIndex)] - ClothSectionVertexRemoveOffset;
-				OriginalIndices.Add(VertexPositionIndex);
-			}
-		}
-
-		struct UIndexProcessedHelper
-		{
-			int32 NewIndex;
-			bool Processed;
-			UIndexProcessedHelper(const int32& Index)
-				: NewIndex(Index)
-				, Processed(false)
-			{
-
-			}
-		};
-
-		TMap<int32, UIndexProcessedHelper> IndexRemapper;
-		{
-			TArray<int32> SortedIndices = OriginalIndices;
-			SortedIndices.Sort();
-
-			TSet<int32> SortedUniqueIndices;
-			SortedUniqueIndices.Append(SortedIndices);
-			SortedIndices.Empty();
-
-			SortedIndices = SortedUniqueIndices.Array();
-
-			for (size_t NewIndex = 0; NewIndex < SortedIndices.Num(); NewIndex++)
-			{
-				IndexRemapper.Add(SortedIndices[NewIndex], NewIndex);
-			}
-		}
-
-		FGLTFIndexArray Indices;
-		FGLTFPositionArray PositionBuffer;
-		FGLTFColorArray VertexColorBuffer;
-		FGLTFNormalArray Normals;
-		FGLTFTangentArray Tangents;
-		TArray<FGLTFUVArray> UVs;
-		TArray<FGLTFJointInfluenceArray> JointInfluences; //Per Group
-		TArray<FGLTFJointWeightArray> JointWeights; //Per Group
-
-		int32 PrimtiveVertexCount = IndexRemapper.Num();
-
-		Indices.Reserve(OriginalIndices.Num());
-
-		PositionBuffer.SetNumZeroed(PrimtiveVertexCount);
-		VertexColorBuffer.SetNumZeroed(PrimtiveVertexCount);
-		Normals.SetNumZeroed(PrimtiveVertexCount);
-		Tangents.SetNumZeroed(PrimtiveVertexCount);
+		FGLTFJsonPrimitive& JsonPrimitive = JsonMesh->Primitives[PrimitiveIndex];
 
-		UVs.AddDefaulted(NumTexCoords);
-		for (size_t UVIndex = 0; UVIndex < NumTexCoords; UVIndex++)
+		//Set  Primitive:
+		JsonPrimitive.Indices = Builder.AddUniqueIndexAccessor(MeshPrimitiveDescription.Indices, SkeletalMesh->GetName() + (MeshDescriptionParser.MeshDetails.NumberOfPrimitives > 1 ? (TEXT("_") + FString::FromInt(PrimitiveIndex)) : TEXT("")));
+		JsonPrimitive.Attributes.Position = Builder.AddUniquePositionAccessor(MeshPrimitiveDescription.Positions);
+		if (!MeshPrimitiveDescription.VertexColors.IsEmpty()) JsonPrimitive.Attributes.Color0 = Builder.AddUniqueColorAccessor(MeshPrimitiveDescription.VertexColors);
+		JsonPrimitive.Attributes.Normal = Builder.AddUniqueNormalAccessor(MeshPrimitiveDescription.Normals);
+		JsonPrimitive.Attributes.Tangent = Builder.AddUniqueTangentAccessor(MeshPrimitiveDescription.Tangents);
+		JsonPrimitive.Attributes.TexCoords.AddUninitialized(MeshDescriptionParser.MeshDetails.UVCount);
+		for (size_t UVIndex = 0; UVIndex < MeshDescriptionParser.MeshDetails.UVCount; UVIndex++)
 		{
-			UVs[UVIndex].SetNumZeroed(PrimtiveVertexCount);
+			JsonPrimitive.Attributes.TexCoords[UVIndex] = Builder.AddUniqueUVAccessor(MeshPrimitiveDescription.UVs[UVIndex]);
 		}
 
-		const uint32 GroupCount = (Section.GetMaxBoneInfluences() + 3) / 4;
 		if (Builder.ExportOptions->bExportVertexSkinWeights)
 		{
-			// TODO: report warning or option to limit groups (of joints and weights) since most viewers don't support more than one?
-			JointInfluences.AddDefaulted(GroupCount);
-			JointWeights.AddDefaulted(GroupCount);
-			for (size_t GroupIndex = 0; GroupIndex < GroupCount; GroupIndex++)
-			{
-				JointInfluences[GroupIndex].SetNumZeroed(PrimtiveVertexCount);
-				JointWeights[GroupIndex].SetNumZeroed(PrimtiveVertexCount);
-			}
-		}
-
-		for (const int32& OriginalIndex : OriginalIndices)
-		{
-			UIndexProcessedHelper& IndexProcessedHelper = IndexRemapper[OriginalIndex];
-			int32 NewIndex = IndexProcessedHelper.NewIndex;
-
-			Indices.Add(NewIndex);
-
-			if (IndexProcessedHelper.Processed)
-			{
-				continue;
-			}
-			IndexProcessedHelper.Processed = true;
-
-			const FSoftSkinVertex& OriginalSoftSkinVertex = SoftSkinVertices[OriginalIndex];
-
-
-			PositionBuffer[NewIndex] = OriginalSoftSkinVertex.Position;
-
-			if (bExportVertexColors)
+			int32 GroupCount = MeshPrimitiveDescription.JointInfluences.Num();
+			if (GroupCount == MeshPrimitiveDescription.JointWeights.Num())
 			{
-				VertexColorBuffer[NewIndex] = OriginalSoftSkinVertex.Color;
-			}
-
-			Normals[NewIndex] = OriginalSoftSkinVertex.TangentZ;
-			Tangents[NewIndex] = OriginalSoftSkinVertex.TangentX;
-
-			for (size_t UVIndex = 0; UVIndex < NumTexCoords; UVIndex++)
-			{
-				UVs[UVIndex][NewIndex] = OriginalSoftSkinVertex.UVs[UVIndex];
-			}
-
-			if (Builder.ExportOptions->bExportVertexSkinWeights)
-			{
-				for (size_t GroupIndex = 0; GroupIndex < GroupCount; GroupIndex++)
+				JsonPrimitive.Attributes.Joints.AddUninitialized(GroupCount);
+				JsonPrimitive.Attributes.Weights.AddUninitialized(GroupCount);
+				for (size_t GroupCountIndex = 0; GroupCountIndex < GroupCount; GroupCountIndex++)
 				{
-					for (size_t GroupBoneInfluenceCounter = 0; GroupBoneInfluenceCounter < 4; GroupBoneInfluenceCounter++)
-					{
-						int8 InfluenceIndex = GroupIndex * 4 + GroupBoneInfluenceCounter;
-
-						JointInfluences[GroupIndex][NewIndex][GroupBoneInfluenceCounter] = Section.BoneMap[OriginalSoftSkinVertex.InfluenceBones[InfluenceIndex]];
-						JointWeights[GroupIndex][NewIndex][GroupBoneInfluenceCounter] = OriginalSoftSkinVertex.InfluenceWeights[InfluenceIndex];
-					}
+					JsonPrimitive.Attributes.Joints[GroupCountIndex] = Builder.AddUniqueJointAccessor(MeshPrimitiveDescription.JointInfluences[GroupCountIndex]);
+					JsonPrimitive.Attributes.Weights[GroupCountIndex] = Builder.AddUniqueWeightAccessor(MeshPrimitiveDescription.JointWeights[GroupCountIndex]);
 				}
 			}
 		}
 
-		if (Tangents.Num() > 0 && DoesBufferHasZeroVector(Tangents))
-		{
-			//Do not Export Tangents list that is zeroed out.
-			Tangents.Reset();
-		}
-
-		if (Normals.Num() > 0 && DoesBufferHasZeroVector(Normals))
-		{
-			//Do not Export Normals list that is zeroed out.
-			Normals.Reset();
-		}
-
-		//Set glTF Primitive:
-		JsonPrimitive.Indices = Builder.AddUniqueIndexAccessor(Indices, SkeletalMesh->GetName() + (SectionCount > 1 ? (TEXT("_") + FString::FromInt(SectionIndex)) : TEXT("")));
-		JsonPrimitive.Attributes.Position = Builder.AddUniquePositionAccessor(PositionBuffer);
-		if (bExportVertexColors) JsonPrimitive.Attributes.Color0 = Builder.AddUniqueColorAccessor(VertexColorBuffer);
-		JsonPrimitive.Attributes.Normal = Builder.AddUniqueNormalAccessor(Normals);
-		JsonPrimitive.Attributes.Tangent = Builder.AddUniqueTangentAccessor(Tangents);
-		JsonPrimitive.Attributes.TexCoords.AddUninitialized(NumTexCoords);
-		for (size_t UVIndex = 0; UVIndex < NumTexCoords; UVIndex++)
-		{
-			JsonPrimitive.Attributes.TexCoords[UVIndex] = Builder.AddUniqueUVAccessor(UVs[UVIndex]);
-		}
-
-		if (Builder.ExportOptions->bExportVertexSkinWeights)
-		{
-			JsonPrimitive.Attributes.Joints.AddUninitialized(GroupCount);
-			JsonPrimitive.Attributes.Weights.AddUninitialized(GroupCount);
-			for (size_t GroupCountIndex = 0; GroupCountIndex < GroupCount; GroupCountIndex++)
-			{
-				JsonPrimitive.Attributes.Joints[GroupCountIndex] = Builder.AddUniqueJointAccessor(JointInfluences[GroupCountIndex]);
-				JsonPrimitive.Attributes.Weights[GroupCountIndex] = Builder.AddUniqueWeightAccessor(JointWeights[GroupCountIndex]);
-			}
-		}
-
 		//
-		const UMaterialInterface* Material = Materials.IsValidIndex(Section.MaterialIndex) ? Materials[Section.MaterialIndex] : MaterialSlots[Section.MaterialIndex].MaterialInterface;
-		JsonPrimitive.Material = Builder.AddUniqueMaterial(Material, MeshData, { SectionIndex });
+		const UMaterialInterface* Material = Materials.IsValidIndex(MeshPrimitiveDescription.MaterialIndex) ? Materials[MeshPrimitiveDescription.MaterialIndex] : MaterialSlots[MeshPrimitiveDescription.MaterialIndex].MaterialInterface;
+		JsonPrimitive.Material = Builder.AddUniqueMaterial(Material, MeshData, { MeshPrimitiveDescription.MaterialIndex });
 
 		//Validations:
 		if (JsonPrimitive.Attributes.Position == nullptr)
@@ -779,6 +481,8 @@ void FGLTFDelayedSkeletalMeshTask::ProcessSourceModel(const TArray<FSkeletalMate
 					*JsonMesh->Name
 				));
 		}
+
+		PrimitiveIndex++;
 	}
 }
 #endif
diff --git a/Engine/Plugins/Enterprise/GLTFExporter/Source/GLTFExporter/Private/Utilities/MeshParsingUtilities.cpp b/Engine/Plugins/Enterprise/GLTFExporter/Source/GLTFExporter/Private/Utilities/MeshParsingUtilities.cpp
new file mode 100644
--- /dev/null
+++ b/Engine/Plugins/Enterprise/GLTFExporter/Source/GLTFExporter/Private/Utilities/MeshParsingUtilities.cpp
@@ -0,0 +1,542 @@
+// Copyright Epic Games, Inc. All Rights Reserved.
+
+#include "Utilities/MeshParsingUtilities.h"
+
+#if WITH_EDITORONLY_DATA
+
+#include "Components/SplineMeshComponent.h"
+#include "Engine/SkeletalMesh.h"
+#include "Engine/StaticMesh.h"
+#include "SkeletalMeshAttributes.h"
+#include "Utilities/MeshAttributesArray.h"
+
+namespace UE::MeshParser
+{
+	namespace Private
+	{
+		template <class T>
+		bool Equals(const T& A, const T& B)
+		{
+			return A.Equals(B, UE_SMALL_NUMBER);
+		}
+
+		bool Equals(const FColor& A, const FColor& B)
+		{
+			return A == B;
+		}
+
+		template <class T>
+		bool CheckCompareSetValue(bool& IsSet, T& StoredValue, const T& ToSetValue)
+		{
+			if (IsSet)
+			{
+				return Equals(StoredValue, ToSetValue);
+			}
+
+			IsSet = true;
+			StoredValue = ToSetValue;
+
+			return true;
+		}
+
+		template <class T>
+		bool DoesBufferHasZeroVector(TArray<T> Buffer, float Tolerance = UE_KINDA_SMALL_NUMBER)
+		{
+			for (const T& Value : Buffer)
+			{
+				if (FMath::Abs(Value.X) <= Tolerance
+					&& FMath::Abs(Value.Y) <= Tolerance
+					&& FMath::Abs(Value.Z) <= Tolerance)
+				{
+					return true;
+				}
+			}
+			return false;
+		}
+	}
+
+	//
+	//FMeshPrimitiveDescription
+	//
+
+	bool FMeshPrimitiveDescription::IsEmpty()
+	{
+		return !(Indices.Num() && Positions.Num());
+	}
+
+	void FMeshPrimitiveDescription::EmptyContainers()
+	{
+		Indices.Empty();
+		Positions.Empty();
+		VertexColors.Empty();
+		Normals.Empty();
+		Tangents.Empty();
+		UVs.Empty();
+		JointInfluences.Empty();
+		JointWeights.Empty();
+	}
+
+	//Except for Joint related containers
+	void FMeshPrimitiveDescription::PrepareContainers(const int32& IndexCount, const int32& AttributesCount, const int32& UVCount, bool bPrepareVertexColors)
+	{
+		Indices.Reserve(IndexCount);
+		Positions.SetNumZeroed(AttributesCount);
+		if (bPrepareVertexColors) VertexColors.SetNumZeroed(AttributesCount);
+		Normals.SetNumZeroed(AttributesCount);
+		Tangents.SetNumZeroed(AttributesCount);
+		UVs.AddDefaulted(UVCount);
+		for (size_t UVIndex = 0; UVIndex < UVCount; UVIndex++)
+		{
+			UVs[UVIndex].SetNumZeroed(AttributesCount);
+		}
+	}
+
+	void FMeshPrimitiveDescription::PrepareJointContainers(const uint32& JointGroupCount, const int32& AttributesCount)
+	{
+		JointInfluences.AddDefaulted(JointGroupCount);
+		JointWeights.AddDefaulted(JointGroupCount);
+		for (size_t GroupIndex = 0; GroupIndex < JointGroupCount; GroupIndex++)
+		{
+			JointInfluences[GroupIndex].SetNumZeroed(AttributesCount);
+			JointWeights[GroupIndex].SetNumZeroed(AttributesCount);
+		}
+	}
+
+	FExportConfigs::FExportConfigs(bool bInExportVertexSkinWeights,
+		bool bInExportVertexColors,
+		const USplineMeshComponent* InSplineMeshComponent,
+		int32 InSkeletonInfluenceCountPerGroup)
+		: bExportVertexSkinWeights(bInExportVertexSkinWeights)
+		, bExportVertexColors(bInExportVertexColors)
+		, SplineMeshComponent(InSplineMeshComponent)
+		, SkeletonInfluenceCountPerGroup(InSkeletonInfluenceCountPerGroup)
+	{
+	}
+
+
+	//
+	//FMeshDescriptionParser
+	//
+	template <class T>
+	FMeshDescriptionParser<T>::FMeshDescriptionParser(const FMeshDescription* InMeshDescription, const TArray<T>& InMaterialSlots)
+		: MeshDescription(InMeshDescription)
+		, MaterialSlots(InMaterialSlots)
+	{
+		FStaticMeshConstAttributes MeshAttributes(*MeshDescription);
+
+		//Set MeshAttrbiutes to be used later:
+		VertexPositions = MeshAttributes.GetVertexPositions();
+		VertexInstanceNormals = MeshAttributes.GetVertexInstanceNormals();
+		VertexInstanceTangents = MeshAttributes.GetVertexInstanceTangents();
+		VertexInstanceBinormalSigns = MeshAttributes.GetVertexInstanceBinormalSigns();
+		VertexInstanceUVs = MeshAttributes.GetVertexInstanceUVs();
+		VertexInstanceColors = MeshAttributes.GetVertexInstanceColors();
+
+		VertexInstanceIdToVertexId = MeshAttributes.GetVertexInstanceVertexIndices().GetRawArray();
+		PolygonGroupMaterialSlotNames = MeshAttributes.GetPolygonGroupMaterialSlotNames();
+
+		//Mesh level details that are used for the parsing
+		MeshDetails.UVCount = VertexInstanceUVs.GetNumChannels();
+		MeshDetails.NumberOfPrimitives = MeshDescription->PolygonGroups().Num();
+		MeshDetails.bHasVertexColors = MeshDescription->VertexInstanceAttributes().HasAttribute(MeshAttribute::VertexInstance::Color);
+	}
+
+	template <class T>
+	void FMeshDescriptionParser<T>::Parse(TArray<FMeshPrimitiveDescription>& OutMeshPrimitiveDescriptions, const FExportConfigs& ExportConfigs)
+	{
+		using namespace Private;
+
+		//Iterate through all PolygonsGroups:
+		for (const FPolygonGroupID& PolygonGroupID : MeshDescription->PolygonGroups().GetElementIDs())
+		{
+			//Acquire the TriangleIDs for the current PolygonGroup:
+			const TArrayView<const FTriangleID>& TriangleIDs = MeshDescription->GetPolygonGroupTriangles(PolygonGroupID);
+
+			//We add the MeshPrimitiveDescription here, even if it is empty, in order to keep the same number of MeshPrimitiveDescriptions as the PolygonGroups.
+			FMeshPrimitiveDescription& MeshPrimitiveDescription = OutMeshPrimitiveDescriptions.AddDefaulted_GetRef();
+
+			if (TriangleIDs.Num() == 0)
+			{
+				//Do not export empty primitives.
+				continue;
+			}
+
+			MeshPrimitiveDescription.MaterialIndex = GetMaterialIndex(PolygonGroupID);
+
+			//VertexIDs of the TriangleIDs can be part of a bigger/unified container
+			//For exporting purposes we want to have the Primitives to have their own containers for Attributes+Indices.
+			//Which means we have still remap the Vertices to their own containers per Primitives.
+
+			if (!ParseVertexBased(TriangleIDs, MeshPrimitiveDescription, ExportConfigs))
+			{
+				//Emtpy all attribute storages in case the ParseVertexBased already started filling them:
+				MeshPrimitiveDescription.EmptyContainers();
+
+				ParseVertexInstanceBased(TriangleIDs, MeshPrimitiveDescription, ExportConfigs);
+			}
+
+			if (MeshPrimitiveDescription.Tangents.Num() > 0 && DoesBufferHasZeroVector(MeshPrimitiveDescription.Tangents))
+			{
+				//Do not Export Tangents list that is zeroed out.
+				MeshPrimitiveDescription.Tangents.Reset();
+			}
+
+			if (MeshPrimitiveDescription.Normals.Num() > 0 && DoesBufferHasZeroVector(MeshPrimitiveDescription.Normals))
+			{
+				//Do not Export Normals list that is zeroed out.
+				MeshPrimitiveDescription.Normals.Reset();
+			}
+		}
+	}
+
+	template <class T>
+	bool FMeshDescriptionParser<T>::ParseVertexBased(const TArrayView<const FTriangleID>& TriangleIDs, FMeshPrimitiveDescription& MeshPrimitiveDescription, const FExportConfigs& ExportConfigs)
+	{
+		using namespace Private;
+
+		const int32 TrianglesCount = TriangleIDs.Num();
+		const int32 VertexInstanceCount = TrianglesCount * 3; //Max number of VertexInstances used in the PolygonGroup aka the Primitive.
+
+		TSet<FVertexID> UniqueVertexIDs;
+		UniqueVertexIDs.Reserve(VertexInstanceCount);
+		for (const FTriangleID& TriangleID : TriangleIDs)
+		{
+			TArrayView<const FVertexInstanceID> TriangleVertexInstanceIDs = MeshDescription->GetTriangleVertexInstances(TriangleID);
+			for (const FVertexInstanceID& VertexInstanceID : TriangleVertexInstanceIDs)
+			{
+				FVertexID VertexID = VertexInstanceIdToVertexId[VertexInstanceID];
+				UniqueVertexIDs.Add(VertexID);
+			}
+		}
+		const int32 VertexCount = UniqueVertexIDs.Num();
+		TMap<int32, int32> OriginalToExportedVertexIdMap;
+
+		auto GetOriginalAndExportedVertexIDs = [&OriginalToExportedVertexIdMap](TConstArrayView<FVertexID>& VertexInstanceIdToVertexId, const FVertexID& VertexInstanceID, FVertexID& VertexID, int32& ExportedVertexID)
+			{
+				if (VertexInstanceIdToVertexId.Num() < VertexInstanceID)
+				{
+					return false;
+				}
+				VertexID = VertexInstanceIdToVertexId[VertexInstanceID];
+				int32* ExportedVertexIDPtr = OriginalToExportedVertexIdMap.Find(VertexID);
+				if (!ExportedVertexIDPtr)
+				{
+					return false;
+				}
+				ExportedVertexID = *ExportedVertexIDPtr;
+
+				return true;
+			};
+
+
+		//This is an attribute set track which we use for identifying if an Attribute[VertexID] has been set already or not.
+		//	Which is needed in order to identify if all VertexInstances of a given VertexID uses the same/identical/equal Values per instance.
+		//	- In case all VertexInstances of given VertexIDs use the same Value we can use the packed version for export.
+		//	- In case that's not true, we fall back onto VertexInstance exports instead.
+		struct FVertexAttributesSetTracker
+		{
+			TArray<bool> VertexColors;
+			TArray<bool> Normals;
+			TArray<bool> Tangents;
+			TArray<TArray<bool>> UVs;
+
+			FVertexAttributesSetTracker(int32 VertexCount, int32 UVCount)
+			{
+				VertexColors.SetNumZeroed(VertexCount);
+				Normals.SetNumZeroed(VertexCount);
+				Tangents.SetNumZeroed(VertexCount);
+				for (size_t UVIndex = 0; UVIndex < UVCount; UVIndex++)
+				{
+					UVs.Add(TArray<bool>());
+					UVs[UVIndex].SetNumZeroed(VertexCount);
+				}
+			}
+		};
+		FVertexAttributesSetTracker AttributesSetTracker(VertexCount, MeshDetails.UVCount);
+
+		const bool bCanExportVertexColors = MeshDetails.bHasVertexColors && ExportConfigs.bExportVertexColors;
+
+		//Set and Reserve the appropriate counts for the Attributes (and Indices)
+		//SetNum/Reserve the containers:
+		MeshPrimitiveDescription.PrepareContainers(VertexInstanceCount, VertexCount, MeshDetails.UVCount, bCanExportVertexColors);
+
+		for (const FTriangleID& TriangleID : TriangleIDs)
+		{
+			TArrayView<const FVertexInstanceID> TriangleVertexInstanceIDs = MeshDescription->GetTriangleVertexInstances(TriangleID);
+			for (const FVertexInstanceID& VertexInstanceID : TriangleVertexInstanceIDs)
+			{
+				FVertexID OriginalVertexID = VertexInstanceIdToVertexId[VertexInstanceID];
+				int32* ExportedVertexIDPtr = OriginalToExportedVertexIdMap.Find(OriginalVertexID);
+				int32 ExportedVertexID = ExportedVertexIDPtr ? *ExportedVertexIDPtr : OriginalToExportedVertexIdMap.Num();
+
+				MeshPrimitiveDescription.Indices.Add(ExportedVertexID);
+
+				//Positions are Vertex based so they can be directly set, no need to check them:
+				if (!ExportedVertexIDPtr)
+				{
+					OriginalToExportedVertexIdMap.Add(OriginalVertexID, ExportedVertexID);
+
+					MeshPrimitiveDescription.Positions[ExportedVertexID] = VertexPositions[OriginalVertexID];
+				}
+
+				//For the VertexInstance based attributes we have to check if all VertexInstances that use the same Vertex base are identical/equal or not.
+				//If they are equal then we can use these packed value, if NOT then we have to use the VertexInstances based solution.
+				if (!CheckCompareSetValue(AttributesSetTracker.Normals[ExportedVertexID], MeshPrimitiveDescription.Normals[ExportedVertexID], VertexInstanceNormals[VertexInstanceID]))
+				{
+					return false;
+				}
+				if (!CheckCompareSetValue(AttributesSetTracker.Tangents[ExportedVertexID], MeshPrimitiveDescription.Tangents[ExportedVertexID], FVector4f(VertexInstanceTangents[VertexInstanceID], VertexInstanceBinormalSigns[VertexInstanceID])))
+				{
+					return false;
+				}
+
+				for (size_t UVIndex = 0; UVIndex < MeshDetails.UVCount; UVIndex++)
+				{
+					if (!CheckCompareSetValue(AttributesSetTracker.UVs[UVIndex][ExportedVertexID], MeshPrimitiveDescription.UVs[UVIndex][ExportedVertexID], VertexInstanceUVs.Get(VertexInstanceID, UVIndex)))
+					{
+						return false;
+					}
+				}
+			}
+		}
+
+		//Set VertexColors if MeshDescription has any:
+		if (bCanExportVertexColors)
+		{
+			for (const FTriangleID& TriangleID : TriangleIDs)
+			{
+				TArrayView<const FVertexInstanceID> TriangleVertexInstanceIDs = MeshDescription->GetTriangleVertexInstances(TriangleID);
+				for (const FVertexInstanceID& VertexInstanceID : TriangleVertexInstanceIDs)
+				{
+					FVertexID VertexID = INDEX_NONE; int32 ExportedVertexID = INDEX_NONE;
+					if (!GetOriginalAndExportedVertexIDs(VertexInstanceIdToVertexId, VertexInstanceID, VertexID, ExportedVertexID))
+					{
+						return false;
+					}
+
+					const FVector4f& SourceVertexColor = VertexInstanceColors[VertexInstanceID];
+					if (!CheckCompareSetValue(AttributesSetTracker.VertexColors[ExportedVertexID], MeshPrimitiveDescription.VertexColors[ExportedVertexID], FLinearColor(SourceVertexColor).ToFColor(true)))
+					{
+						return false;
+					}
+				}
+			}
+		}
+
+		if (ExportConfigs.bExportVertexSkinWeights)
+		{
+			FSkeletalMeshConstAttributes SkeletalMeshAttributes(*MeshDescription);
+			FSkinWeightsVertexAttributesConstRef VertexSkinWeights = SkeletalMeshAttributes.GetVertexSkinWeights();
+
+			//Find MaxBoneOInfluence straight from VertexSkinWeights:
+			int32 MaxBoneInfluences = 0;
+			for (const FVertexID VertexID : MeshDescription->Vertices().GetElementIDs())
+			{
+				const FVertexBoneWeightsConst BoneWeights = VertexSkinWeights.Get(VertexID);
+				if (MaxBoneInfluences < BoneWeights.Num())
+				{
+					MaxBoneInfluences = BoneWeights.Num();
+				}
+			}
+
+			const uint32 JointGroupCount = (MaxBoneInfluences + 3) / 4;
+			MeshPrimitiveDescription.PrepareJointContainers(JointGroupCount, VertexCount);
+
+			for (const FTriangleID& TriangleID : TriangleIDs)
+			{
+				TArrayView<const FVertexInstanceID> TriangleVertexInstanceIDs = MeshDescription->GetTriangleVertexInstances(TriangleID);
+				for (const FVertexInstanceID& VertexInstanceID : TriangleVertexInstanceIDs)
+				{
+					FVertexID VertexID = INDEX_NONE; int32 ExportedVertexID = INDEX_NONE;
+					if (!GetOriginalAndExportedVertexIDs(VertexInstanceIdToVertexId, VertexInstanceID, VertexID, ExportedVertexID))
+					{
+						return false;
+					}
+
+					int32 GroupCounter = 0;
+					int32 NumInfluences = 0;
+					int32 InGroupInfluenceCounter = 0;
+					const FVertexBoneWeightsConst BoneWeights = VertexSkinWeights.Get(VertexID);
+					for (const UE::AnimationCore::FBoneWeight BoneWeight : BoneWeights)
+					{
+						MeshPrimitiveDescription.JointInfluences[GroupCounter][ExportedVertexID][InGroupInfluenceCounter] = BoneWeight.GetBoneIndex();
+						MeshPrimitiveDescription.JointWeights[GroupCounter][ExportedVertexID][InGroupInfluenceCounter] = BoneWeight.GetRawWeight();
+
+						NumInfluences++;
+						GroupCounter = (NumInfluences / ExportConfigs.SkeletonInfluenceCountPerGroup);
+						InGroupInfluenceCounter = NumInfluences % ExportConfigs.SkeletonInfluenceCountPerGroup;
+					}
+				}
+			}
+		}
+
+		return true;
+	}
+
+	template <class T>
+	void FMeshDescriptionParser<T>::ParseVertexInstanceBased(const TArrayView<const FTriangleID>& TriangleIDs, FMeshPrimitiveDescription& MeshPrimitiveDescription, const FExportConfigs& ExportConfigs)
+	{
+		//This approach should be used when the VertexBased fails.
+		//		VertexBased mesh parsing can fail due to non idential VertexInstance attributes (per Vertex)
+
+		using namespace Private;
+
+		const int32 TrianglesCount = TriangleIDs.Num();
+		const int32 MaxVertexInstanceCount = TrianglesCount * 3; //Max number of VertexInstances used in the PolygonGroup aka the Primitive.
+
+		//Note: while for the VertexBasedMeshParsing we calculate the VERTEX Count at this stage, here we calculate the VERTEXINSTANCE count:
+		TSet<FVertexID> UniqueVertexInstanceIDs;
+		UniqueVertexInstanceIDs.Reserve(MaxVertexInstanceCount);
+		for (const FTriangleID& TriangleID : TriangleIDs)
+		{
+			TArrayView<const FVertexInstanceID> TriangleVertexInstanceIDs = MeshDescription->GetTriangleVertexInstances(TriangleID);
+			for (const FVertexInstanceID& VertexInstanceID : TriangleVertexInstanceIDs)
+			{
+				UniqueVertexInstanceIDs.Add(VertexInstanceID);
+			}
+		}
+		const int32 VertexInstanceCount = UniqueVertexInstanceIDs.Num();
+		TMap<int32, int32> OriginalToExportedVertexIdMap;
+		auto GetOriginalAndExportedVertexIDs = [&OriginalToExportedVertexIdMap](const FVertexID& VertexInstanceID, int32& ExportedVertexID)
+			{
+				int32* ExportedVertexIDPtr = OriginalToExportedVertexIdMap.Find(VertexInstanceID);
+				if (!ExportedVertexIDPtr)
+				{
+					return false;
+				}
+				ExportedVertexID = *ExportedVertexIDPtr;
+
+				return true;
+			};
+
+		const bool bCanExportVertexColors = MeshDetails.bHasVertexColors && ExportConfigs.bExportVertexColors;
+
+		//SetNum/Reserve the containers:
+		MeshPrimitiveDescription.PrepareContainers(VertexInstanceCount, VertexInstanceCount, MeshDetails.UVCount, bCanExportVertexColors);
+
+		//Based on Triangles acquire the VertexInstanceID, based on which acquire the VertexInstance Attributes:
+		for (const FTriangleID& TriangleID : TriangleIDs)
+		{
+			TArrayView<const FVertexInstanceID> TriangleVertexInstanceIDs = MeshDescription->GetTriangleVertexInstances(TriangleID);
+			for (const FVertexInstanceID& VertexInstanceID : TriangleVertexInstanceIDs)
+			{
+				FVertexID OriginalVertexID = VertexInstanceIdToVertexId[VertexInstanceID];
+				int32* ExportedVertexIDPtr = OriginalToExportedVertexIdMap.Find(VertexInstanceID);
+				int32 ExportedVertexID = ExportedVertexIDPtr ? *ExportedVertexIDPtr : OriginalToExportedVertexIdMap.Num();
+
+				MeshPrimitiveDescription.Indices.Add(ExportedVertexID);
+
+				//Positions and Indicies are Vertex based so they can be directly set, no need to check them:
+				if (!ExportedVertexIDPtr)
+				{
+					OriginalToExportedVertexIdMap.Add(VertexInstanceID, ExportedVertexID);
+
+					MeshPrimitiveDescription.Positions[ExportedVertexID] = VertexPositions[OriginalVertexID];
+
+					MeshPrimitiveDescription.Normals[ExportedVertexID] = VertexInstanceNormals[VertexInstanceID];
+					MeshPrimitiveDescription.Tangents[ExportedVertexID] = FVector4f(VertexInstanceTangents[VertexInstanceID], VertexInstanceBinormalSigns[VertexInstanceID]);
+
+					for (size_t UVIndex = 0; UVIndex < MeshDetails.UVCount; UVIndex++)
+					{
+						MeshPrimitiveDescription.UVs[UVIndex][ExportedVertexID] = VertexInstanceUVs.Get(VertexInstanceID, UVIndex);
+					}
+				}
+			}
+		}
+
+		//Set VertexColors if MeshDescription has any:
+		if (bCanExportVertexColors)
+		{
+			for (const FTriangleID& TriangleID : TriangleIDs)
+			{
+				TArrayView<const FVertexInstanceID> TriangleVertexInstanceIDs = MeshDescription->GetTriangleVertexInstances(TriangleID);
+				for (const FVertexInstanceID& VertexInstanceID : TriangleVertexInstanceIDs)
+				{
+					int32 ExportedVertexID = INDEX_NONE;
+					if (!GetOriginalAndExportedVertexIDs(VertexInstanceID, ExportedVertexID)) return;
+
+					const FVector4f& SourceVertexColor = VertexInstanceColors[VertexInstanceID];
+
+					MeshPrimitiveDescription.VertexColors[ExportedVertexID] = FLinearColor(SourceVertexColor).ToFColor(true);
+				}
+			}
+		}
+
+		if (ExportConfigs.bExportVertexSkinWeights)
+		{
+			FSkeletalMeshConstAttributes SkeletalMeshAttributes(*MeshDescription);
+			FSkinWeightsVertexAttributesConstRef VertexSkinWeights = SkeletalMeshAttributes.GetVertexSkinWeights();
+
+			int32 MaxBoneInfluences = 0;
+			for (const FVertexID VertexID : MeshDescription->Vertices().GetElementIDs())
+			{
+				const FVertexBoneWeightsConst BoneWeights = VertexSkinWeights.Get(VertexID);
+				if (MaxBoneInfluences < BoneWeights.Num())
+				{
+					MaxBoneInfluences = BoneWeights.Num();
+				}
+			}
+			const uint32 JointGroupCount = (MaxBoneInfluences + 3) / 4;
+
+			MeshPrimitiveDescription.PrepareJointContainers(JointGroupCount, VertexInstanceCount);
+
+			for (const FTriangleID& TriangleID : TriangleIDs)
+			{
+				TArrayView<const FVertexInstanceID> TriangleVertexInstanceIDs = MeshDescription->GetTriangleVertexInstances(TriangleID);
+				for (const FVertexInstanceID& VertexInstanceID : TriangleVertexInstanceIDs)
+				{
+					int32 ExportedVertexID = INDEX_NONE;
+					if (!GetOriginalAndExportedVertexIDs(VertexInstanceID, ExportedVertexID)) return;
+					FVertexID OriginalVertexID = VertexInstanceIdToVertexId[VertexInstanceID];
+
+					int32 GroupCounter = 0;
+					int32 NumInfluences = 0;
+					int32 InGroupInfluenceCounter = 0;
+					const FVertexBoneWeightsConst BoneWeights = VertexSkinWeights.Get(OriginalVertexID);
+					for (const UE::AnimationCore::FBoneWeight BoneWeight : BoneWeights)
+					{
+						MeshPrimitiveDescription.JointInfluences[GroupCounter][ExportedVertexID][InGroupInfluenceCounter] = BoneWeight.GetBoneIndex();
+						MeshPrimitiveDescription.JointWeights[GroupCounter][ExportedVertexID][InGroupInfluenceCounter] = BoneWeight.GetRawWeight();
+
+						NumInfluences++;
+						GroupCounter = (NumInfluences / ExportConfigs.SkeletonInfluenceCountPerGroup);
+						InGroupInfluenceCounter = NumInfluences % ExportConfigs.SkeletonInfluenceCountPerGroup;
+					}
+				}
+			}
+		}
+	}
+
+	template <class T>
+	int32 FMeshDescriptionParser<T>::GetMaterialIndex(const FPolygonGroupID& PolygonGroupID)
+	{
+		FName MaterialSlotName = PolygonGroupMaterialSlotNames[PolygonGroupID];
+
+		int32 MaterialIndex = INDEX_NONE;
+
+		for (size_t MatIndex = 0; MatIndex < MaterialSlots.Num(); MatIndex++)
+		{
+			if (MaterialSlots[MatIndex].ImportedMaterialSlotName == MaterialSlotName)
+			{
+				MaterialIndex = MatIndex;
+				break;
+			}
+		}
+
+		if (MaterialIndex == INDEX_NONE)
+		{
+			MaterialIndex = PolygonGroupID.GetValue();
+		}
+		if (!MaterialSlots.IsValidIndex(MaterialIndex))
+		{
+			MaterialIndex = 0;
+		}
+
+		return MaterialIndex;
+	}
+
+	template struct FMeshDescriptionParser<FStaticMaterial>;
+	template struct FMeshDescriptionParser<FSkeletalMaterial>;
+}
+
+#endif
\ No newline at end of file
diff --git a/Engine/Plugins/Enterprise/GLTFExporter/Source/GLTFExporter/Public/Converters/GLTFMeshAttributesArray.h b/Engine/Plugins/Enterprise/GLTFExporter/Source/GLTFExporter/Public/Converters/GLTFMeshAttributesArray.h
--- a/Engine/Plugins/Enterprise/GLTFExporter/Source/GLTFExporter/Public/Converters/GLTFMeshAttributesArray.h
+++ b/Engine/Plugins/Enterprise/GLTFExporter/Source/GLTFExporter/Public/Converters/GLTFMeshAttributesArray.h
@@ -4,31 +4,14 @@
 
 #include "CoreMinimal.h"
 #include "BoneIndices.h"
+#include "Utilities/MeshAttributesArray.h"
 
-template <typename T>
-class FGLTFMeshAttributesArray : public TArray<T>
-{
-public:
+typedef FMeshAttributesArray<int32> FGLTFIndexArray;
+typedef FMeshAttributesArray<FVector3f> FGLTFPositionArray;
+typedef FMeshAttributesArray<FVector3f> FGLTFNormalArray;
+typedef FMeshAttributesArray<FVector4f> FGLTFTangentArray;
+typedef FMeshAttributesArray<FVector2f> FGLTFUVArray;
+typedef FMeshAttributesArray<FColor> FGLTFColorArray;
 
-	using TArray<T>::TArray;
-	using TArray<T>::operator=;
-
-	friend uint32 GetTypeHash(const FGLTFMeshAttributesArray& AttributesArray)
-	{
-		uint32 Hash = GetTypeHash(AttributesArray.Num());
-		for (const T& Attribute : AttributesArray)
-		{
-			Hash = HashCombine(Hash, GetTypeHash(Attribute));
-		}
-		return Hash;
-	}
-};
-
-typedef FGLTFMeshAttributesArray<FColor> FGLTFColorArray;
-typedef FGLTFMeshAttributesArray<int32> FGLTFIndexArray;
-typedef FGLTFMeshAttributesArray<UE::Math::TIntVector4<FBoneIndexType>> FGLTFJointInfluenceArray;
-typedef FGLTFMeshAttributesArray<UE::Math::TIntVector4<uint16>> FGLTFJointWeightArray;
-typedef FGLTFMeshAttributesArray<FVector3f> FGLTFNormalArray;
-typedef FGLTFMeshAttributesArray<FVector3f> FGLTFPositionArray;
-typedef FGLTFMeshAttributesArray<FVector4f> FGLTFTangentArray;
-typedef FGLTFMeshAttributesArray<FVector2f> FGLTFUVArray;
\ No newline at end of file
+typedef FMeshAttributesArray<UE::Math::TIntVector4<FBoneIndexType>> FGLTFJointInfluenceArray;
+typedef FMeshAttributesArray<UE::Math::TIntVector4<uint16>> FGLTFJointWeightArray;
\ No newline at end of file
diff --git a/Engine/Plugins/Enterprise/GLTFExporter/Source/GLTFExporter/Public/Utilities/MeshAttributesArray.h b/Engine/Plugins/Enterprise/GLTFExporter/Source/GLTFExporter/Public/Utilities/MeshAttributesArray.h
new file mode 100644
--- /dev/null
+++ b/Engine/Plugins/Enterprise/GLTFExporter/Source/GLTFExporter/Public/Utilities/MeshAttributesArray.h
@@ -0,0 +1,24 @@
+// Copyright Epic Games, Inc. All Rights Reserved.
+
+#pragma once
+
+#include "Containers/Array.h"
+
+template <typename T>
+class FMeshAttributesArray : public TArray<T>
+{
+public:
+
+	using TArray<T>::TArray;
+	using TArray<T>::operator=;
+
+	friend uint32 GetTypeHash(const FMeshAttributesArray& AttributesArray)
+	{
+		uint32 Hash = GetTypeHash(AttributesArray.Num());
+		for (const T& Attribute : AttributesArray)
+		{
+			Hash = HashCombine(Hash, GetTypeHash(Attribute));
+		}
+		return Hash;
+	}
+};
\ No newline at end of file
diff --git a/Engine/Plugins/Enterprise/GLTFExporter/Source/GLTFExporter/Public/Utilities/MeshParsingUtilities.h b/Engine/Plugins/Enterprise/GLTFExporter/Source/GLTFExporter/Public/Utilities/MeshParsingUtilities.h
new file mode 100644
--- /dev/null
+++ b/Engine/Plugins/Enterprise/GLTFExporter/Source/GLTFExporter/Public/Utilities/MeshParsingUtilities.h
@@ -0,0 +1,120 @@
+// Copyright Epic Games, Inc. All Rights Reserved.
+
+#pragma once
+
+#include "Misc/CoreMiscDefines.h"
+
+#if WITH_EDITORONLY_DATA
+
+#include "BoneIndices.h"
+#include "MeshDescription.h" //for mesh attribute containers
+#include "Utilities/MeshAttributesArray.h"
+
+class USplineMeshComponent;
+
+namespace UE::MeshParser
+{
+	struct FMeshPrimitiveDescription
+	{
+		int32                                   MaterialIndex = INDEX_NONE;
+
+		FMeshAttributesArray<int32>             Indices;
+		FMeshAttributesArray<FVector3f>         Positions;
+		FMeshAttributesArray<FVector3f>         Normals;
+		FMeshAttributesArray<FVector4f>         Tangents;
+		TArray<FMeshAttributesArray<FVector2f>> UVs;
+
+		FMeshAttributesArray<FColor>            VertexColors;
+
+		TArray<FMeshAttributesArray<UE::Math::TIntVector4<FBoneIndexType>>> JointInfluences; //TODO: Per Group (make group size configurable)
+		TArray<FMeshAttributesArray<UE::Math::TIntVector4<uint16>>>         JointWeights; //TODO: Per Group (make group size configurable)
+
+		bool IsEmpty();
+		void EmptyContainers();
+
+		//Except for Joint related containers
+		void PrepareContainers(const int32& IndexCount, const int32& AttributesCount, const int32& UVCount, bool bHasVertexColors);
+
+		void PrepareJointContainers(const uint32& JointGroupCount, const int32& AttributesCount);
+	};
+
+	struct FMeshDetails
+	{
+		int32 NumberOfPrimitives = 0;
+		int32 UVCount = 0;
+		bool bHasVertexColors = false;
+	};
+
+	struct FExportConfigs
+	{
+		bool bExportVertexSkinWeights = false;
+		bool bExportVertexColors = false;
+		const USplineMeshComponent* SplineMeshComponent = nullptr;
+		int32 SkeletonInfluenceCountPerGroup = INT32_MAX;
+
+		FExportConfigs(bool bInExportVertexSkinWeights,
+			bool bInExportVertexColors,
+			const USplineMeshComponent* InSplineMeshComponent,
+			int32 InSkeletonInfluenceCountPerGroup);
+	};
+
+	/*
+	* Note0: We are compacting the Vertex Attributes and Indices per Primitives, with no overlapping.
+	* Note1: Primary approach tries compaction on a Vertex bases, if that fails (due to VertexInstance Attributes differing on a per Vertex base), it falls back onto VertexInstance based compaction.
+	*			Parsing does not do compaction for identical Vertex attribute values.
+	*			It's either driven by FVertexIDs or FVertexInstanceIDs within Indices.
+	* Note2: TODO: perhaps add a solution path where we keep the unfication across primitives?
+	* class T: Static/Skeletal Material Slots
+	*/
+	template <class T>
+	struct FMeshDescriptionParser
+	{
+		FMeshDescriptionParser(const FMeshDescription* InMeshDescription, const TArray<T>& InMaterialSlots);
+
+		/*
+		* Parses FMeshDescription into the FMeshPrimitiveMeshDescriptions array based on the provided ExportConfigs.
+		*/
+		void Parse(TArray<FMeshPrimitiveDescription>& OutMeshPrimitiveDescriptions, const FExportConfigs& ExportConfigs);
+
+	private:
+		/*
+		* Tries to parse the PolygonGroup to Primitive, using only the Vertices.
+		* Checks if the VertexInstances use identical values per Vertices.
+		* If any of the VertexInstances of a given Vertex do not match Values we fail out and return false.
+		*	(in which case we are falling back on using ParseVertexInstanceBased function)
+		*/
+		bool ParseVertexBased(const TArrayView<const FTriangleID>& TriangleIDs, FMeshPrimitiveDescription& MeshPrimitiveDescription, const FExportConfigs& ExportConfigs);
+
+		/*
+		* Parses the PolygonGroup to Primitive, using the VertexInstances.
+		* Each VertexInstance will be a new entry in the VertexAttribute containers (aka a new 'Vertex' in the parsed data set).
+		*/
+		void ParseVertexInstanceBased(const TArrayView<const FTriangleID>& TriangleIDs, FMeshPrimitiveDescription& MeshPrimitiveDescription, const FExportConfigs& ExportConfigs);
+
+		/*
+		* Gets the Material Index for the PolygonGroupID, based on the MeshDescription and the MaterialSlots provided in the struct's constructor.
+		*/
+		int32 GetMaterialIndex(const FPolygonGroupID& PolygonGroupID);
+
+		//Inputs:
+		const FMeshDescription* MeshDescription; //Primarily used for GetPolygonGroupTriangles, GetTriangleVertexInstances
+		const TArray<T>& MaterialSlots;
+
+		//Helper variable acquired from Inputs:
+		TVertexAttributesConstRef<FVector3f> VertexPositions;
+		TVertexInstanceAttributesConstRef<FVector3f> VertexInstanceNormals;
+		TVertexInstanceAttributesConstRef<FVector3f> VertexInstanceTangents;
+		TVertexInstanceAttributesConstRef<float> VertexInstanceBinormalSigns;
+		TVertexInstanceAttributesConstRef<FVector2f> VertexInstanceUVs;
+		TVertexInstanceAttributesConstRef<FVector4f> VertexInstanceColors;
+
+		TConstArrayView<FVertexID> VertexInstanceIdToVertexId;
+		TPolygonGroupAttributesConstRef<FName> PolygonGroupMaterialSlotNames;
+
+	public:
+		//MeshDetails like UVCount/NumberOfPrimitives/bHasVertexColor for ease of access:
+		FMeshDetails MeshDetails;
+	};
+}
+
+#endif
\ No newline at end of file
