[glTF Exporter] CIS fix for variable name shadowing in MeshParsingUtilities.

#rnx

diff --git a/Engine/Plugins/Enterprise/GLTFExporter/Source/GLTFExporter/Private/Utilities/MeshParsingUtilities.cpp b/Engine/Plugins/Enterprise/GLTFExporter/Source/GLTFExporter/Private/Utilities/MeshParsingUtilities.cpp
--- a/Engine/Plugins/Enterprise/GLTFExporter/Source/GLTFExporter/Private/Utilities/MeshParsingUtilities.cpp
+++ b/Engine/Plugins/Enterprise/GLTFExporter/Source/GLTFExporter/Private/Utilities/MeshParsingUtilities.cpp
@@ -211,13 +211,13 @@ namespace UE::MeshParser
 		const int32 VertexCount = UniqueVertexIDs.Num();
 		TMap<int32, int32> OriginalToExportedVertexIdMap;
 
-		auto GetOriginalAndExportedVertexIDs = [&OriginalToExportedVertexIdMap](TConstArrayView<FVertexID>& VertexInstanceIdToVertexId, const FVertexID& VertexInstanceID, FVertexID& VertexID, int32& ExportedVertexID)
+		auto GetOriginalAndExportedVertexIDs = [&OriginalToExportedVertexIdMap](TConstArrayView<FVertexID>& LocalVertexInstanceIdToVertexId, const FVertexID& VertexInstanceID, FVertexID& VertexID, int32& ExportedVertexID)
 			{
-				if (VertexInstanceIdToVertexId.Num() < VertexInstanceID)
+				if (LocalVertexInstanceIdToVertexId.Num() < VertexInstanceID)
 				{
 					return false;
 				}
-				VertexID = VertexInstanceIdToVertexId[VertexInstanceID];
+				VertexID = LocalVertexInstanceIdToVertexId[VertexInstanceID];
 				int32* ExportedVertexIDPtr = OriginalToExportedVertexIdMap.Find(VertexID);
 				if (!ExportedVertexIDPtr)
 				{
