[glTF Exporter] Re-adding the EPIC_lightmap_textures extension.

#rnx

diff --git a/Engine/Plugins/Enterprise/GLTFExporter/Source/GLTFExporter/Private/Builders/GLTFConvertBuilder.cpp b/Engine/Plugins/Enterprise/GLTFExporter/Source/GLTFExporter/Private/Builders/GLTFConvertBuilder.cpp
--- a/Engine/Plugins/Enterprise/GLTFExporter/Source/GLTFExporter/Private/Builders/GLTFConvertBuilder.cpp
+++ b/Engine/Plugins/Enterprise/GLTFExporter/Source/GLTFExporter/Private/Builders/GLTFConvertBuilder.cpp
@@ -402,6 +402,16 @@ FGLTFJsonTexture* FGLTFConvertBuilder::AddUniqueTexture(const UTextureRenderTarg
 	return TextureRenderTarget2DConverter->GetOrAdd(Texture, bToSRGB);
 }
 
+FGLTFJsonTexture* FGLTFConvertBuilder::AddUniqueTexture(const ULightMapTexture2D* Texture)
+{
+	if (Texture == nullptr)
+	{
+		return nullptr;
+	}
+
+	return TextureLightMapConverter->GetOrAdd(Texture);
+}
+
 FGLTFJsonImage* FGLTFConvertBuilder::AddUniqueImage(TGLTFSharedArray<FColor>& Pixels, FIntPoint Size, bool bIgnoreAlpha, const FString& Name)
 {
 	return ImageConverter->GetOrAdd(Name, bIgnoreAlpha, Size, Pixels);
@@ -611,6 +621,17 @@ FGLTFJsonMaterialVariant* FGLTFConvertBuilder::AddUniqueMaterialVariant(const UV
 	return MaterialVariantConverter->GetOrAdd(Variant);
 }
 
+FGLTFJsonLightMap* FGLTFConvertBuilder::AddUniqueLightMap(const UStaticMeshComponent* StaticMeshComponent)
+{
+	if (StaticMeshComponent == nullptr)
+	{
+		return nullptr;
+	}
+
+	return LightMapConverter->GetOrAdd(StaticMeshComponent);
+}
+
+
 void FGLTFConvertBuilder::RegisterObjectVariant(const UObject* Object, const UPropertyValue* Property)
 {
 	TArray<const UPropertyValue*>& Variants = ObjectVariants.FindOrAdd(Object);
diff --git a/Engine/Plugins/Enterprise/GLTFExporter/Source/GLTFExporter/Private/Builders/GLTFJsonBuilder.cpp b/Engine/Plugins/Enterprise/GLTFExporter/Source/GLTFExporter/Private/Builders/GLTFJsonBuilder.cpp
--- a/Engine/Plugins/Enterprise/GLTFExporter/Source/GLTFExporter/Private/Builders/GLTFJsonBuilder.cpp
+++ b/Engine/Plugins/Enterprise/GLTFExporter/Source/GLTFExporter/Private/Builders/GLTFJsonBuilder.cpp
@@ -103,6 +103,11 @@ FGLTFJsonLight* FGLTFJsonBuilder::AddLight()
 	return JsonRoot.Lights.Add();
 }
 
+FGLTFJsonLightMap* FGLTFJsonBuilder::AddLightMap()
+{
+	return JsonRoot.LightMaps.Add();
+}
+
 FGLTFJsonLightIES* FGLTFJsonBuilder::AddLightIES()
 {
 	return JsonRoot.LightIESs.Add();
diff --git a/Engine/Plugins/Enterprise/GLTFExporter/Source/GLTFExporter/Private/Converters/GLTFLightMapConverters.cpp b/Engine/Plugins/Enterprise/GLTFExporter/Source/GLTFExporter/Private/Converters/GLTFLightMapConverters.cpp
new file mode 100644
--- /dev/null
+++ b/Engine/Plugins/Enterprise/GLTFExporter/Source/GLTFExporter/Private/Converters/GLTFLightMapConverters.cpp
@@ -0,0 +1,81 @@
+// Copyright Epic Games, Inc. All Rights Reserved.
+
+#include "Converters/GLTFLightMapConverters.h"
+#include "Converters/GLTFMeshUtilities.h"
+#include "Builders/GLTFContainerBuilder.h"
+#include "Components/StaticMeshComponent.h"
+#include "Engine/SkinnedAssetCommon.h"
+#include "Engine/StaticMesh.h"
+#include "Engine/MapBuildDataRegistry.h"
+#include "StaticMeshComponentLODInfo.h"
+#include "StaticMeshResources.h"
+
+FGLTFJsonLightMap* FGLTFLightMapConverter::Convert(const UStaticMeshComponent* StaticMeshComponent)
+{
+	if (Builder.ExportOptions->TextureImageFormat == EGLTFTextureImageFormat::None)
+	{
+		return nullptr;
+	}
+
+	const UStaticMesh* StaticMesh = StaticMeshComponent->GetStaticMesh();
+
+	if (StaticMesh == nullptr)
+	{
+		return nullptr;
+	}
+
+	const int32 LODIndex = FGLTFMeshUtilities::GetLOD(StaticMesh, StaticMeshComponent, Builder.ExportOptions->DefaultLevelOfDetail);
+	const FStaticMeshLODResources& LODResources = StaticMesh->GetLODForExport(LODIndex);
+	const int32 CoordinateIndex = StaticMesh->GetLightMapCoordinateIndex();
+
+	if (CoordinateIndex < 0 || CoordinateIndex >= LODResources.GetNumTexCoords())
+	{
+		return nullptr;
+	}
+
+	constexpr int32 LightMapLODIndex = 0; // TODO: why is this zero?
+
+	if (!StaticMeshComponent->LODData.IsValidIndex(LightMapLODIndex))
+	{
+		return nullptr;
+	}
+
+	const FStaticMeshComponentLODInfo& ComponentLODInfo = StaticMeshComponent->LODData[LightMapLODIndex];
+	const FMeshMapBuildData* MeshMapBuildData = StaticMeshComponent->GetMeshMapBuildData(ComponentLODInfo);
+
+	if (MeshMapBuildData == nullptr || MeshMapBuildData->LightMap == nullptr)
+	{
+		return nullptr;
+	}
+
+	const FLightMap2D* LightMap2D = MeshMapBuildData->LightMap->GetLightMap2D();
+
+	if (LightMap2D == nullptr)
+	{
+		return nullptr;
+	}
+
+	const FLightMapInteraction LightMapInteraction = LightMap2D->GetInteraction(GMaxRHIFeatureLevel);
+	const ULightMapTexture2D* Texture = LightMapInteraction.GetTexture(true);
+	FGLTFJsonTexture* JsonTexture = Builder.AddUniqueTexture(Texture);
+
+	if (JsonTexture == nullptr)
+	{
+		return nullptr;
+	}
+
+	const FVector2f CoordinateBias = FVector2f(LightMap2D->GetCoordinateBias());
+	const FVector2f CoordinateScale = FVector2f(LightMap2D->GetCoordinateScale());
+	const FVector4f& LightMapAdd = LightMapInteraction.GetAddArray()[0];
+	const FVector4f& LightMapScale = LightMapInteraction.GetScaleArray()[0];
+
+	FGLTFJsonLightMap* JsonLightMap = Builder.AddLightMap();
+	StaticMeshComponent->GetName(JsonLightMap->Name); // TODO: use better name (similar to light and camera)
+	JsonLightMap->Texture.Index = JsonTexture;
+	JsonLightMap->Texture.TexCoord = CoordinateIndex;
+	JsonLightMap->LightMapScale = { LightMapScale.X, LightMapScale.Y, LightMapScale.Z, LightMapScale.W };
+	JsonLightMap->LightMapAdd = { LightMapAdd.X, LightMapAdd.Y, LightMapAdd.Z, LightMapAdd.W };
+	JsonLightMap->CoordinateScaleBias = { CoordinateScale.X, CoordinateScale.Y, CoordinateBias.X, CoordinateBias.Y };
+
+	return JsonLightMap;
+}
diff --git a/Engine/Plugins/Enterprise/GLTFExporter/Source/GLTFExporter/Private/Converters/GLTFNodeConverters.cpp b/Engine/Plugins/Enterprise/GLTFExporter/Source/GLTFExporter/Private/Converters/GLTFNodeConverters.cpp
--- a/Engine/Plugins/Enterprise/GLTFExporter/Source/GLTFExporter/Private/Converters/GLTFNodeConverters.cpp
+++ b/Engine/Plugins/Enterprise/GLTFExporter/Source/GLTFExporter/Private/Converters/GLTFNodeConverters.cpp
@@ -185,6 +185,11 @@ void FGLTFComponentConverter::ConvertComponentSpecialization(const USceneCompone
 				Node->Children.Add(InstanceNode);
 			}
 		}
+
+		if (Builder.ExportOptions->bExportLightmaps)
+		{
+			Node->LightMap = Builder.AddUniqueLightMap(InstancedMeshComp);
+		}
 	}
 	else if (const USplineMeshComponent* SplineMeshComponent = Cast<USplineMeshComponent>(SceneComponent))
 	{
@@ -247,6 +252,11 @@ void FGLTFComponentConverter::ConvertComponentSpecialization(const USceneCompone
 	else if (const UStaticMeshComponent* StaticMeshComponent = Cast<UStaticMeshComponent>(SceneComponent))
 	{
 		Node->Mesh = Builder.AddUniqueMesh(StaticMeshComponent);
+
+		if (Builder.ExportOptions->bExportLightmaps)
+		{
+			Node->LightMap = Builder.AddUniqueLightMap(StaticMeshComponent);
+		}
 	}
 	else if (const USkeletalMeshComponent* SkeletalMeshComponent = Cast<USkeletalMeshComponent>(SceneComponent))
 	{
diff --git a/Engine/Plugins/Enterprise/GLTFExporter/Source/GLTFExporter/Private/Converters/GLTFTextureConverters.cpp b/Engine/Plugins/Enterprise/GLTFExporter/Source/GLTFExporter/Private/Converters/GLTFTextureConverters.cpp
--- a/Engine/Plugins/Enterprise/GLTFExporter/Source/GLTFExporter/Private/Converters/GLTFTextureConverters.cpp
+++ b/Engine/Plugins/Enterprise/GLTFExporter/Source/GLTFExporter/Private/Converters/GLTFTextureConverters.cpp
@@ -26,3 +26,19 @@ FGLTFJsonTexture* FGLTFTextureRenderTarget2DConverter::Convert(const UTextureRen
 	Builder.ScheduleSlowTask<FGLTFDelayedTextureRenderTarget2DTask>(Builder, RenderTarget2D, bToSRGB, JsonTexture);
 	return JsonTexture;
 }
+
+FGLTFJsonTexture* FGLTFTextureLightMapConverter::Convert(const ULightMapTexture2D* LightMap)
+{
+	if (Builder.ExportOptions->TextureImageFormat == EGLTFTextureImageFormat::None)
+	{
+		return nullptr;
+	}
+
+#if WITH_EDITOR
+	FGLTFJsonTexture* JsonTexture = Builder.AddTexture();
+	Builder.ScheduleSlowTask<FGLTFDelayedTextureLightMapTask>(Builder, LightMap, JsonTexture);
+	return JsonTexture;
+#else
+	return nullptr;
+#endif
+}
diff --git a/Engine/Plugins/Enterprise/GLTFExporter/Source/GLTFExporter/Private/Json/GLTFJsonLightMap.cpp b/Engine/Plugins/Enterprise/GLTFExporter/Source/GLTFExporter/Private/Json/GLTFJsonLightMap.cpp
new file mode 100644
--- /dev/null
+++ b/Engine/Plugins/Enterprise/GLTFExporter/Source/GLTFExporter/Private/Json/GLTFJsonLightMap.cpp
@@ -0,0 +1,20 @@
+// Copyright Epic Games, Inc. All Rights Reserved.
+
+#include "Json/GLTFJsonLightMap.h"
+
+void FGLTFJsonLightMap::WriteObject(IGLTFJsonWriter& Writer) const
+{
+	if (!Name.IsEmpty())
+	{
+		Writer.Write(TEXT("name"), Name);
+	}
+
+	if (Texture.Index != nullptr)
+	{
+		Writer.Write(TEXT("texture"), Texture);
+	}
+
+	Writer.Write(TEXT("lightmapScale"), LightMapScale);
+	Writer.Write(TEXT("lightmapAdd"), LightMapAdd);
+	Writer.Write(TEXT("coordinateScaleBias"), CoordinateScaleBias);
+}
diff --git a/Engine/Plugins/Enterprise/GLTFExporter/Source/GLTFExporter/Private/Json/GLTFJsonNode.cpp b/Engine/Plugins/Enterprise/GLTFExporter/Source/GLTFExporter/Private/Json/GLTFJsonNode.cpp
--- a/Engine/Plugins/Enterprise/GLTFExporter/Source/GLTFExporter/Private/Json/GLTFJsonNode.cpp
+++ b/Engine/Plugins/Enterprise/GLTFExporter/Source/GLTFExporter/Private/Json/GLTFJsonNode.cpp
@@ -5,6 +5,7 @@
 #include "Json/GLTFJsonSkin.h"
 #include "Json/GLTFJsonMesh.h"
 #include "Json/GLTFJsonLight.h"
+#include "Json/GLTFJsonLightMap.h"
 
 void FGLTFJsonNode::WriteObject(IGLTFJsonWriter& Writer) const
 {
@@ -31,7 +32,8 @@ void FGLTFJsonNode::WriteObject(IGLTFJsonWriter& Writer) const
 	}
 
 	if (Light != nullptr
-		|| (LightIESInstance != nullptr && LightIESInstance->HasValue()))
+		|| (LightIESInstance != nullptr && LightIESInstance->HasValue())
+		|| LightMap != nullptr)
 	{
 		Writer.StartExtensions();
 
@@ -41,6 +43,14 @@ void FGLTFJsonNode::WriteObject(IGLTFJsonWriter& Writer) const
 			Writer.Write(TEXT("light"), Light);
 			Writer.EndExtension();
 		}
+
+		if (LightMap != nullptr)
+		{
+			Writer.StartExtension(EGLTFJsonExtension::EPIC_LightmapTextures);
+			Writer.Write(TEXT("lightmap"), LightMap);
+			Writer.EndExtension();
+		}
+
 		
 		if ((LightIESInstance != nullptr && LightIESInstance->HasValue())) {
 			Writer.StartExtension(EGLTFJsonExtension::EXT_LightsIES);
diff --git a/Engine/Plugins/Enterprise/GLTFExporter/Source/GLTFExporter/Private/Json/GLTFJsonRoot.cpp b/Engine/Plugins/Enterprise/GLTFExporter/Source/GLTFExporter/Private/Json/GLTFJsonRoot.cpp
--- a/Engine/Plugins/Enterprise/GLTFExporter/Source/GLTFExporter/Private/Json/GLTFJsonRoot.cpp
+++ b/Engine/Plugins/Enterprise/GLTFExporter/Source/GLTFExporter/Private/Json/GLTFJsonRoot.cpp
@@ -25,7 +25,10 @@ void FGLTFJsonRoot::WriteObject(IGLTFJsonWriter& Writer) const
 	if (Skins.Num() > 0) Writer.Write(TEXT("skins"), Skins);
 	if (Textures.Num() > 0) Writer.Write(TEXT("textures"), Textures);
 
-	if (Lights.Num() > 0 || MaterialVariants.Num() > 0)
+	if (MaterialVariants.Num() > 0
+		|| Lights.Num() > 0 
+		|| LightMaps.Num() > 0
+		|| LightIESs.Num() > 0)
 	{
 		Writer.StartExtensions();
 
@@ -43,6 +46,13 @@ void FGLTFJsonRoot::WriteObject(IGLTFJsonWriter& Writer) const
 			Writer.EndExtension();
 		}
 
+		if (LightMaps.Num() > 0)
+		{
+			Writer.StartExtension(EGLTFJsonExtension::EPIC_LightmapTextures);
+			Writer.Write(TEXT("lightmaps"), LightMaps);
+			Writer.EndExtension();
+		}
+
 		if (LightIESs.Num() > 0)
 		{
 			Writer.StartExtension(EGLTFJsonExtension::EXT_LightsIES);
diff --git a/Engine/Plugins/Enterprise/GLTFExporter/Source/GLTFExporter/Private/Options/GLTFExportOptions.cpp b/Engine/Plugins/Enterprise/GLTFExporter/Source/GLTFExporter/Private/Options/GLTFExportOptions.cpp
--- a/Engine/Plugins/Enterprise/GLTFExporter/Source/GLTFExporter/Private/Options/GLTFExportOptions.cpp
+++ b/Engine/Plugins/Enterprise/GLTFExporter/Source/GLTFExporter/Private/Options/GLTFExportOptions.cpp
@@ -37,6 +37,7 @@ void UGLTFExportOptions::ResetToDefault()
 	TextureImageFormat = EGLTFTextureImageFormat::PNG;
 	TextureImageQuality = 0;
 	bExportTextureTransforms = true;
+	bExportLightmaps = false;
 	bAdjustNormalmaps = true;
 	bExportHiddenInGame = false;
 	bExportLights = true;
diff --git a/Engine/Plugins/Enterprise/GLTFExporter/Source/GLTFExporter/Private/Tasks/GLTFDelayedTextureTasks.cpp b/Engine/Plugins/Enterprise/GLTFExporter/Source/GLTFExporter/Private/Tasks/GLTFDelayedTextureTasks.cpp
--- a/Engine/Plugins/Enterprise/GLTFExporter/Source/GLTFExporter/Private/Tasks/GLTFDelayedTextureTasks.cpp
+++ b/Engine/Plugins/Enterprise/GLTFExporter/Source/GLTFExporter/Private/Tasks/GLTFDelayedTextureTasks.cpp
@@ -4,11 +4,16 @@
 #include "Builders/GLTFContainerBuilder.h"
 #include "Utilities/GLTFCoreUtilities.h"
 #include "Converters/GLTFTextureUtilities.h"
+#include "Converters/GLTFNameUtilities.h"
 #include "Engine/Texture2D.h"
 #include "Engine/TextureRenderTarget2D.h"
 
 FString FGLTFDelayedTexture2DTask::GetName()
 {
+	if (!Texture2D)
+	{
+		return TEXT("");
+	}
 	return Texture2D->GetName();
 }
 
@@ -56,6 +61,10 @@ void FGLTFDelayedTexture2DTask::Process()
 
 FString FGLTFDelayedTextureRenderTarget2DTask::GetName()
 {
+	if (!RenderTarget2D)
+	{
+		return TEXT("");
+	}
 	return RenderTarget2D->GetName();
 }
 
@@ -85,3 +94,46 @@ void FGLTFDelayedTextureRenderTarget2DTask::Process()
 	JsonTexture->Source = Builder.AddUniqueImage(Pixels, Size, bIgnoreAlpha, JsonTexture->Name);
 	JsonTexture->Sampler = Builder.AddUniqueSampler(RenderTarget2D);
 }
+
+#if WITH_EDITOR
+
+FString FGLTFDelayedTextureLightMapTask::GetName()
+{
+	if (!LightMap)
+	{
+		return TEXT("");
+	}
+	return LightMap->GetName();
+}
+
+void FGLTFDelayedTextureLightMapTask::Process()
+{
+	FGLTFTextureUtilities::FullyLoad(LightMap);
+	LightMap->GetName(JsonTexture->Name);
+
+	const FIntPoint Size = FGLTFTextureUtilities::GetInGameSize(LightMap);
+	UTextureRenderTarget2D* RenderTarget = FGLTFTextureUtilities::CreateRenderTarget(Size, false);
+
+	// TODO: preserve maximum image quality (avoid compression artifacts) by copying source data (and adjustments) to a temp texture
+	FGLTFTextureUtilities::DrawTexture(RenderTarget, LightMap);
+
+	TGLTFSharedArray<FColor> Pixels;
+	if (!FGLTFTextureUtilities::ReadPixels(RenderTarget, *Pixels))
+	{
+		Builder.LogWarning(FString::Printf(TEXT("Failed to read pixels for 2D texture %s"), *JsonTexture->Name));
+		return;
+	}
+
+	if (Builder.ExportOptions->bAdjustNormalmaps && LightMap->IsNormalMap())
+	{
+		// TODO: add support for adjusting normals in GLTFNormalMapPreview instead
+		FGLTFTextureUtilities::FlipGreenChannel(*Pixels);
+	}
+
+	const bool bIgnoreAlpha = FGLTFTextureUtilities::IsAlphaless(LightMap->GetPixelFormat());
+
+	JsonTexture->Source = Builder.AddUniqueImage(Pixels, Size, bIgnoreAlpha, JsonTexture->Name);
+	JsonTexture->Sampler = Builder.AddUniqueSampler(LightMap->GetTextureAddressX(), LightMap->GetTextureAddressX(), LightMap->Filter, LightMap->LODGroup);
+}
+
+#endif
diff --git a/Engine/Plugins/Enterprise/GLTFExporter/Source/GLTFExporter/Private/Tasks/GLTFDelayedTextureTasks.h b/Engine/Plugins/Enterprise/GLTFExporter/Source/GLTFExporter/Private/Tasks/GLTFDelayedTextureTasks.h
--- a/Engine/Plugins/Enterprise/GLTFExporter/Source/GLTFExporter/Private/Tasks/GLTFDelayedTextureTasks.h
+++ b/Engine/Plugins/Enterprise/GLTFExporter/Source/GLTFExporter/Private/Tasks/GLTFDelayedTextureTasks.h
@@ -27,11 +27,11 @@ public:
 private:
 
 	FGLTFConvertBuilder& Builder;
-	const UTexture2D* Texture2D;
+	const UTexture2D* Texture2D = nullptr;
 	bool bToSRGB;
 	TextureAddress TextureAddressX;
 	TextureAddress TextureAddressY;
-	FGLTFJsonTexture* JsonTexture;
+	FGLTFJsonTexture* JsonTexture = nullptr;
 };
 
 class FGLTFDelayedTextureRenderTarget2DTask : public FGLTFDelayedTask
@@ -54,7 +54,34 @@ public:
 private:
 
 	FGLTFConvertBuilder& Builder;
-	const UTextureRenderTarget2D* RenderTarget2D;
+	const UTextureRenderTarget2D* RenderTarget2D = nullptr;
 	bool bToSRGB;
-	FGLTFJsonTexture* JsonTexture;
+	FGLTFJsonTexture* JsonTexture = nullptr;
 };
+
+#if WITH_EDITOR
+
+class FGLTFDelayedTextureLightMapTask : public FGLTFDelayedTask
+{
+public:
+
+	FGLTFDelayedTextureLightMapTask(FGLTFConvertBuilder& Builder, const ULightMapTexture2D* LightMap, FGLTFJsonTexture* JsonTexture)
+		: FGLTFDelayedTask(EGLTFTaskPriority::Texture)
+		, Builder(Builder)
+		, LightMap(LightMap)
+		, JsonTexture(JsonTexture)
+	{
+	}
+
+	virtual FString GetName() override;
+
+	virtual void Process() override;
+
+private:
+
+	FGLTFConvertBuilder& Builder;
+	const ULightMapTexture2D* LightMap = nullptr;
+	FGLTFJsonTexture* JsonTexture = nullptr;
+};
+
+#endif
diff --git a/Engine/Plugins/Enterprise/GLTFExporter/Source/GLTFExporter/Private/Utilities/GLTFJsonUtilities.cpp b/Engine/Plugins/Enterprise/GLTFExporter/Source/GLTFExporter/Private/Utilities/GLTFJsonUtilities.cpp
--- a/Engine/Plugins/Enterprise/GLTFExporter/Source/GLTFExporter/Private/Utilities/GLTFJsonUtilities.cpp
+++ b/Engine/Plugins/Enterprise/GLTFExporter/Source/GLTFExporter/Private/Utilities/GLTFJsonUtilities.cpp
@@ -21,6 +21,7 @@ const TCHAR* FGLTFJsonUtilities::GetValue(EGLTFJsonExtension Enum)
 		case EGLTFJsonExtension::KHR_MeshQuantization:                return TEXT("KHR_mesh_quantization");
 		case EGLTFJsonExtension::KHR_TextureTransform:                return TEXT("KHR_texture_transform");
 		case EGLTFJsonExtension::KHR_MaterialsSpecular:	              return TEXT("KHR_materials_specular");
+		case EGLTFJsonExtension::EPIC_LightmapTextures:	              return TEXT("EPIC_lightmap_textures");
 		default:
 			checkNoEntry();
 			return TEXT("");
diff --git a/Engine/Plugins/Enterprise/GLTFExporter/Source/GLTFExporter/Public/Builders/GLTFConvertBuilder.h b/Engine/Plugins/Enterprise/GLTFExporter/Source/GLTFExporter/Public/Builders/GLTFConvertBuilder.h
--- a/Engine/Plugins/Enterprise/GLTFExporter/Source/GLTFExporter/Public/Builders/GLTFConvertBuilder.h
+++ b/Engine/Plugins/Enterprise/GLTFExporter/Source/GLTFExporter/Public/Builders/GLTFConvertBuilder.h
@@ -18,6 +18,7 @@
 #include "Converters/GLTFLightConverters.h"
 #include "Converters/GLTFMaterialVariantConverters.h"
 #include "Converters/GLTFMeshAttributesArray.h"
+#include "Converters/GLTFLightMapConverters.h"
 
 class UMeshComponent;
 class UPropertyValue;
@@ -78,6 +79,8 @@ public:
 	FGLTFJsonTexture* AddUniqueTexture(const UTexture* Texture, bool bToSRGB, TextureAddress TextureAddressX = TextureAddress::TA_MAX, TextureAddress TextureAddressY = TextureAddress::TA_MAX);
 	FGLTFJsonTexture* AddUniqueTexture(const UTexture2D* Texture, bool bToSRGB, TextureAddress TextureAddressX, TextureAddress TextureAddressY);
 	FGLTFJsonTexture* AddUniqueTexture(const UTextureRenderTarget2D* Texture, bool bToSRGB);
+	FGLTFJsonTexture* AddUniqueTexture(const ULightMapTexture2D* Texture);
+
 	FGLTFJsonImage* AddUniqueImage(TGLTFSharedArray<FColor>& Pixels, FIntPoint Size, bool bIgnoreAlpha, const FString& Name);
 
 	FGLTFJsonSkin* AddUniqueSkin(FGLTFJsonNode* RootNode, const USkeletalMesh* SkeletalMesh);
@@ -99,8 +102,8 @@ public:
 	FGLTFJsonLight* AddUniqueLight(const ULightComponent* LightComponent);
 	FGLTFJsonLightIES* AddUniqueLightIES(const ULightComponent* LightComponent);
 	FGLTFJsonLightIESInstance* AddUniqueLightIESInstance(const ULightComponent* LightComponent);
-
 	FGLTFJsonMaterialVariant* AddUniqueMaterialVariant(const UVariant* Variant);
+	FGLTFJsonLightMap* AddUniqueLightMap(const UStaticMeshComponent* StaticMeshComponent);
 
 	void RegisterObjectVariant(const UObject* Object, const UPropertyValue* Property);
 	const TArray<const UPropertyValue*>* GetObjectVariants(const UObject* Object) const;
@@ -162,6 +165,10 @@ public:
 	TUniquePtr<FGLTFLightIESInstanceConverter> LightIESInstanceConverter = MakeUnique<FGLTFLightIESInstanceConverter>(*this);
 	TUniquePtr<IGLTFMaterialVariantConverter> MaterialVariantConverter = MakeUnique<FGLTFMaterialVariantConverter>(*this);
 
+	//LightMap related:
+	TUniquePtr<IGLTFTextureLightMapConverter> TextureLightMapConverter = MakeUnique<FGLTFTextureLightMapConverter>(*this);
+	TUniquePtr<IGLTFLightMapConverter> LightMapConverter = MakeUnique<FGLTFLightMapConverter>(*this);
+
 private:
 
 	TMap<const UObject*, TArray<const UPropertyValue*>> ObjectVariants;
diff --git a/Engine/Plugins/Enterprise/GLTFExporter/Source/GLTFExporter/Public/Builders/GLTFJsonBuilder.h b/Engine/Plugins/Enterprise/GLTFExporter/Source/GLTFExporter/Public/Builders/GLTFJsonBuilder.h
--- a/Engine/Plugins/Enterprise/GLTFExporter/Source/GLTFExporter/Public/Builders/GLTFJsonBuilder.h
+++ b/Engine/Plugins/Enterprise/GLTFExporter/Source/GLTFExporter/Public/Builders/GLTFJsonBuilder.h
@@ -29,6 +29,7 @@ public:
 	FGLTFJsonSkin* AddSkin();
 	FGLTFJsonTexture* AddTexture();
 	FGLTFJsonLight* AddLight();
+	FGLTFJsonLightMap* AddLightMap();
 	FGLTFJsonLightIES* AddLightIES();
 	FGLTFJsonLightIESInstance* AddLightIESInstance();
 	FGLTFJsonMaterialVariant* AddMaterialVariant();
diff --git a/Engine/Plugins/Enterprise/GLTFExporter/Source/GLTFExporter/Public/Converters/GLTFLightMapConverters.h b/Engine/Plugins/Enterprise/GLTFExporter/Source/GLTFExporter/Public/Converters/GLTFLightMapConverters.h
new file mode 100644
--- /dev/null
+++ b/Engine/Plugins/Enterprise/GLTFExporter/Source/GLTFExporter/Public/Converters/GLTFLightMapConverters.h
@@ -0,0 +1,22 @@
+// Copyright Epic Games, Inc. All Rights Reserved.
+
+#pragma once
+
+#include "Json/GLTFJsonCore.h"
+#include "Converters/GLTFConverter.h"
+#include "Converters/GLTFBuilderContext.h"
+
+class UStaticMeshComponent;
+
+typedef TGLTFConverter<FGLTFJsonLightMap*, const UStaticMeshComponent*> IGLTFLightMapConverter;
+
+class GLTFEXPORTER_API FGLTFLightMapConverter : public FGLTFBuilderContext, public IGLTFLightMapConverter
+{
+public:
+
+	using FGLTFBuilderContext::FGLTFBuilderContext;
+
+protected:
+
+	virtual FGLTFJsonLightMap* Convert(const UStaticMeshComponent* StaticMeshComponent) override;
+};
diff --git a/Engine/Plugins/Enterprise/GLTFExporter/Source/GLTFExporter/Public/Converters/GLTFTextureConverters.h b/Engine/Plugins/Enterprise/GLTFExporter/Source/GLTFExporter/Public/Converters/GLTFTextureConverters.h
--- a/Engine/Plugins/Enterprise/GLTFExporter/Source/GLTFExporter/Public/Converters/GLTFTextureConverters.h
+++ b/Engine/Plugins/Enterprise/GLTFExporter/Source/GLTFExporter/Public/Converters/GLTFTextureConverters.h
@@ -10,11 +10,13 @@
 
 class UTexture2D;
 class UTextureRenderTarget2D;
+class ULightMapTexture2D;
 
 // TODO: generalize parameter bToSRGB to SamplerType (including normalmap unpacking)
 
 typedef TGLTFConverter<FGLTFJsonTexture*, const UTexture2D*, bool, TextureAddress, TextureAddress> IGLTFTexture2DConverter;
 typedef TGLTFConverter<FGLTFJsonTexture*, const UTextureRenderTarget2D*, bool> IGLTFTextureRenderTarget2DConverter;
+typedef TGLTFConverter<FGLTFJsonTexture*, const ULightMapTexture2D*> IGLTFTextureLightMapConverter;
 
 class GLTFEXPORTER_API FGLTFTexture2DConverter : public FGLTFBuilderContext, public IGLTFTexture2DConverter
 {
@@ -37,3 +39,14 @@ protected:
 
 	virtual FGLTFJsonTexture* Convert(const UTextureRenderTarget2D* RenderTarget2D, bool bToSRGB) override;
 };
+
+class GLTFEXPORTER_API FGLTFTextureLightMapConverter : public FGLTFBuilderContext, public IGLTFTextureLightMapConverter
+{
+public:
+
+	using FGLTFBuilderContext::FGLTFBuilderContext;
+
+protected:
+
+	virtual FGLTFJsonTexture* Convert(const ULightMapTexture2D* LightMap) override;
+};
diff --git a/Engine/Plugins/Enterprise/GLTFExporter/Source/GLTFExporter/Public/Json/GLTFJsonEnums.h b/Engine/Plugins/Enterprise/GLTFExporter/Source/GLTFExporter/Public/Json/GLTFJsonEnums.h
--- a/Engine/Plugins/Enterprise/GLTFExporter/Source/GLTFExporter/Public/Json/GLTFJsonEnums.h
+++ b/Engine/Plugins/Enterprise/GLTFExporter/Source/GLTFExporter/Public/Json/GLTFJsonEnums.h
@@ -18,7 +18,8 @@ enum class EGLTFJsonExtension
 	KHR_MaterialsAnisotropy,
 	KHR_MeshQuantization,
 	KHR_TextureTransform,
-	KHR_MaterialsSpecular
+	KHR_MaterialsSpecular,
+	EPIC_LightmapTextures
 };
 
 enum class EGLTFJsonShadingModel
diff --git a/Engine/Plugins/Enterprise/GLTFExporter/Source/GLTFExporter/Public/Json/GLTFJsonLightMap.h b/Engine/Plugins/Enterprise/GLTFExporter/Source/GLTFExporter/Public/Json/GLTFJsonLightMap.h
new file mode 100644
--- /dev/null
+++ b/Engine/Plugins/Enterprise/GLTFExporter/Source/GLTFExporter/Public/Json/GLTFJsonLightMap.h
@@ -0,0 +1,29 @@
+// Copyright Epic Games, Inc. All Rights Reserved.
+
+#pragma once
+
+#include "Json/GLTFJsonCore.h"
+#include "Json/GLTFJsonMaterial.h"
+
+struct GLTFEXPORTER_API FGLTFJsonLightMap : IGLTFJsonIndexedObject
+{
+	FString              Name;
+	FGLTFJsonTextureInfo Texture;
+	FGLTFJsonVector4     LightMapScale;
+	FGLTFJsonVector4     LightMapAdd;
+	FGLTFJsonVector4     CoordinateScaleBias;
+
+	virtual void WriteObject(IGLTFJsonWriter& Writer) const override;
+
+protected:
+
+	friend TGLTFJsonIndexedObjectArray<FGLTFJsonLightMap, void>;
+
+	FGLTFJsonLightMap(int32 Index)
+		: IGLTFJsonIndexedObject(Index)
+		, LightMapScale(FGLTFJsonVector4::One)
+		, LightMapAdd(FGLTFJsonVector4::Zero)
+		, CoordinateScaleBias({ 1, 1, 0, 0 })
+	{
+	}
+};
diff --git a/Engine/Plugins/Enterprise/GLTFExporter/Source/GLTFExporter/Public/Json/GLTFJsonNode.h b/Engine/Plugins/Enterprise/GLTFExporter/Source/GLTFExporter/Public/Json/GLTFJsonNode.h
--- a/Engine/Plugins/Enterprise/GLTFExporter/Source/GLTFExporter/Public/Json/GLTFJsonNode.h
+++ b/Engine/Plugins/Enterprise/GLTFExporter/Source/GLTFExporter/Public/Json/GLTFJsonNode.h
@@ -13,6 +13,7 @@ struct GLTFEXPORTER_API FGLTFJsonNode : FGLTFJsonTransform, IGLTFJsonIndexedObje
 	FGLTFJsonSkin*   Skin;
 	FGLTFJsonMesh*   Mesh;
 	FGLTFJsonLight*  Light;
+	FGLTFJsonLightMap* LightMap;
 	FGLTFJsonLightIESInstance*  LightIESInstance;
 
 	TArray<FGLTFJsonNode*> Children;
@@ -29,6 +30,7 @@ protected:
 		, Skin(nullptr)
 		, Mesh(nullptr)
 		, Light(nullptr)
+		, LightMap(nullptr)
 		, LightIESInstance(nullptr)
 	{
 	}
diff --git a/Engine/Plugins/Enterprise/GLTFExporter/Source/GLTFExporter/Public/Json/GLTFJsonRoot.h b/Engine/Plugins/Enterprise/GLTFExporter/Source/GLTFExporter/Public/Json/GLTFJsonRoot.h
--- a/Engine/Plugins/Enterprise/GLTFExporter/Source/GLTFExporter/Public/Json/GLTFJsonRoot.h
+++ b/Engine/Plugins/Enterprise/GLTFExporter/Source/GLTFExporter/Public/Json/GLTFJsonRoot.h
@@ -18,6 +18,7 @@
 #include "Json/GLTFJsonSkin.h"
 #include "Json/GLTFJsonTexture.h"
 #include "Json/GLTFJsonLight.h"
+#include "Json/GLTFJsonLightMap.h"
 #include "Json/GLTFJsonMaterialVariant.h"
 
 struct GLTFEXPORTER_API FGLTFJsonRoot : IGLTFJsonObject
@@ -43,6 +44,7 @@ struct GLTFEXPORTER_API FGLTFJsonRoot : IGLTFJsonObject
 	TGLTFJsonIndexedObjectArray<FGLTFJsonSkin>       Skins;
 	TGLTFJsonIndexedObjectArray<FGLTFJsonTexture>    Textures;
 	TGLTFJsonIndexedObjectArray<FGLTFJsonLight>      Lights;
+	TGLTFJsonIndexedObjectArray<FGLTFJsonLightMap>   LightMaps;
 	TGLTFJsonIndexedObjectArray<FGLTFJsonLightIES>   LightIESs;
 	TGLTFJsonIndexedObjectArray<FGLTFJsonLightIESInstance>  LightIESInstances;
 	TGLTFJsonIndexedObjectArray<FGLTFJsonMaterialVariant>   MaterialVariants;
diff --git a/Engine/Plugins/Enterprise/GLTFExporter/Source/GLTFExporter/Public/Options/GLTFExportOptions.h b/Engine/Plugins/Enterprise/GLTFExporter/Source/GLTFExporter/Public/Options/GLTFExportOptions.h
--- a/Engine/Plugins/Enterprise/GLTFExporter/Source/GLTFExporter/Public/Options/GLTFExportOptions.h
+++ b/Engine/Plugins/Enterprise/GLTFExporter/Source/GLTFExporter/Public/Options/GLTFExportOptions.h
@@ -145,6 +145,10 @@ class GLTFEXPORTER_API UGLTFExportOptions : public UObject
 	UPROPERTY(EditAnywhere, BlueprintReadWrite, Config, Category = Texture, Meta = (EditCondition = "TextureImageFormat != EGLTFTextureImageFormat::None"))
 	bool bExportTextureTransforms;
 
+	/** If enabled, export lightmaps (created by Lightmass) when exporting a level. Uses extension EPIC_lightmap_textures, which is supported by Unreal's glTF viewer. */
+	UPROPERTY(EditAnywhere, BlueprintReadWrite, Config, Category = Texture, Meta = (EditCondition = "TextureImageFormat != EGLTFTextureImageFormat::None"))
+	bool bExportLightmaps;
+
 	/** If enabled, exported normalmaps will be adjusted from Unreal to glTF convention (i.e. the green channel is flipped). */
 	UPROPERTY(EditAnywhere, BlueprintReadWrite, Config, Category = Texture, Meta = (EditCondition = "TextureImageFormat != EGLTFTextureImageFormat::None"))
 	bool bAdjustNormalmaps;
