Fixed ULevel nullptr deference when exporting a level with ULevelInstanceComponent's with no loaded level

#13046

diff --git a/Engine/Plugins/Enterprise/GLTFExporter/Source/GLTFExporter/Private/Converters/GLTFNodeConverters.cpp b/Engine/Plugins/Enterprise/GLTFExporter/Source/GLTFExporter/Private/Converters/GLTFNodeConverters.cpp
--- a/Engine/Plugins/Enterprise/GLTFExporter/Source/GLTFExporter/Private/Converters/GLTFNodeConverters.cpp
+++ b/Engine/Plugins/Enterprise/GLTFExporter/Source/GLTFExporter/Private/Converters/GLTFNodeConverters.cpp
@@ -220,8 +220,12 @@ void FGLTFComponentConverter::ConvertComponentSpecialization(const USceneCompone
 	else if (const ULevelInstanceComponent* LevelInstanceComponent = Cast<ULevelInstanceComponent>(SceneComponent))
 	{
 		const ALevelInstance* LevelInstance = Cast<ALevelInstance>(Owner);
-		
-		TArray<TObjectPtr<AActor>> Actors = LevelInstance->GetLoadedLevel()->Actors;
+
+		TArray<TObjectPtr<AActor>> Actors;
+		if (const ULevel* LoadedLevel = LevelInstance->GetLoadedLevel())
+		{
+			Actors = LoadedLevel->Actors;
+		}
 
 		for (const AActor* LevelActor : Actors)
 		{
