[glTF Exporter] Material Constant value read fixes.

#rnx

diff --git a/Engine/Plugins/Enterprise/GLTFExporter/Source/GLTFExporter/Private/Tasks/GLTFDelayedMaterialTasks.cpp b/Engine/Plugins/Enterprise/GLTFExporter/Source/GLTFExporter/Private/Tasks/GLTFDelayedMaterialTasks.cpp
--- a/Engine/Plugins/Enterprise/GLTFExporter/Source/GLTFExporter/Private/Tasks/GLTFDelayedMaterialTasks.cpp
+++ b/Engine/Plugins/Enterprise/GLTFExporter/Source/GLTFExporter/Private/Tasks/GLTFDelayedMaterialTasks.cpp
@@ -1273,7 +1273,7 @@ bool FGLTFDelayedMaterialTask::TryGetConstantColor(FLinearColor& OutValue, const
 		return false;
 	}
 
-	const FMaterialInput<FColor>* MaterialInput = FGLTFMaterialUtilities::GetInputForProperty<FColor>(Material, Property);
+	const FMaterialInput<FLinearColor>* MaterialInput = FGLTFMaterialUtilities::GetInputForProperty<FLinearColor>(Material, Property);
 	if (MaterialInput == nullptr)
 	{
 		// TODO: report error
