[glTFExporter] MorphTarget Support.

#rnx

diff --git a/Engine/Plugins/Enterprise/GLTFExporter/Source/GLTFExporter/Private/Builders/GLTFConvertBuilder.cpp b/Engine/Plugins/Enterprise/GLTFExporter/Source/GLTFExporter/Private/Builders/GLTFConvertBuilder.cpp
--- a/Engine/Plugins/Enterprise/GLTFExporter/Source/GLTFExporter/Private/Builders/GLTFConvertBuilder.cpp
+++ b/Engine/Plugins/Enterprise/GLTFExporter/Source/GLTFExporter/Private/Builders/GLTFConvertBuilder.cpp
@@ -27,14 +27,14 @@ bool FGLTFConvertBuilder::IsRootActor(const AActor* Actor) const
 	return ParentActor == nullptr || !IsSelectedActor(ParentActor);
 }
 
-FGLTFJsonAccessor* FGLTFConvertBuilder::AddUniquePositionAccessor(const FGLTFMeshSection* MeshSection, const FPositionVertexBuffer* VertexBuffer)
+FGLTFJsonAccessor* FGLTFConvertBuilder::AddUniquePositionAccessor(const FGLTFMeshSection* MeshSection, const FPositionVertexBuffer* VertexBuffer, const TMap<uint32, FVector3f>* TargetPositionDeltas)
 {
 	if (VertexBuffer == nullptr)
 	{
 		return nullptr;
 	}
 
-	return PositionBufferConverter->GetOrAdd(MeshSection, VertexBuffer);
+	return PositionBufferConverter->GetOrAdd(MeshSection, VertexBuffer, TargetPositionDeltas);
 }
 
 FGLTFJsonAccessor* FGLTFConvertBuilder::AddUniquePositionAccessor(const FGLTFPositionArray& VertexBuffer)
@@ -57,14 +57,14 @@ FGLTFJsonAccessor* FGLTFConvertBuilder::AddUniqueColorAccessor(const FGLTFColorA
 	return ColorBufferConverterRaw->GetOrAdd(VertexColorBuffer);
 }
 
-FGLTFJsonAccessor* FGLTFConvertBuilder::AddUniqueNormalAccessor(const FGLTFMeshSection* MeshSection, const FStaticMeshVertexBuffer* VertexBuffer)
+FGLTFJsonAccessor* FGLTFConvertBuilder::AddUniqueNormalAccessor(const FGLTFMeshSection* MeshSection, const FStaticMeshVertexBuffer* VertexBuffer, const TMap<uint32, FVector3f>* TargetNormalDeltas)
 {
 	if (VertexBuffer == nullptr)
 	{
 		return nullptr;
 	}
 
-	return NormalBufferConverter->GetOrAdd(MeshSection, VertexBuffer);
+	return NormalBufferConverter->GetOrAdd(MeshSection, VertexBuffer, TargetNormalDeltas);
 }
 
 FGLTFJsonAccessor* FGLTFConvertBuilder::AddUniqueNormalAccessor(const FGLTFNormalArray& Normals)
diff --git a/Engine/Plugins/Enterprise/GLTFExporter/Source/GLTFExporter/Private/Converters/GLTFAccessorConverters.cpp b/Engine/Plugins/Enterprise/GLTFExporter/Source/GLTFExporter/Private/Converters/GLTFAccessorConverters.cpp
--- a/Engine/Plugins/Enterprise/GLTFExporter/Source/GLTFExporter/Private/Converters/GLTFAccessorConverters.cpp
+++ b/Engine/Plugins/Enterprise/GLTFExporter/Source/GLTFExporter/Private/Converters/GLTFAccessorConverters.cpp
@@ -18,6 +18,11 @@ namespace
 		{
 			return FGLTFCoreUtilities::ConvertNormal(Normal);
 		}
+
+		static DestinationType Convert(const SourceType& Normal, const FVector3f& Offset)
+		{
+			return FGLTFCoreUtilities::ConvertNormal(Normal, Offset);
+		}
 	};
 
 	template <typename SourceType>
@@ -27,6 +32,11 @@ namespace
 		{
 			return FGLTFCoreUtilities::ConvertNormal(Normal.ToFVector3f());
 		}
+
+		static FGLTFVector3 Convert(const SourceType& Normal, const FVector3f& Offset)
+		{
+			return FGLTFCoreUtilities::ConvertNormal(Normal.ToFVector3f(), Offset);
+		}
 	};
 
 	template <typename DestinationType, typename SourceType>
@@ -92,7 +102,7 @@ namespace
 	}
 }
 
-FGLTFJsonAccessor* FGLTFPositionBufferConverter::Convert(const FGLTFMeshSection* MeshSection, const FPositionVertexBuffer* VertexBuffer)
+FGLTFJsonAccessor* FGLTFPositionBufferConverter::Convert(const FGLTFMeshSection* MeshSection, const FPositionVertexBuffer* VertexBuffer, const TMap<uint32, FVector3f>* TargetPositionDeltas)
 {
 	if (VertexBuffer == nullptr || VertexBuffer->GetNumVertices() == 0)
 	{
@@ -115,25 +125,59 @@ FGLTFJsonAccessor* FGLTFPositionBufferConverter::Convert(const FGLTFMeshSection*
 	TArray<FGLTFVector3> Positions;
 	Positions.AddUninitialized(VertexCount);
 
-	if (IndexMap.Num() > 0)
+	if (TargetPositionDeltas && TargetPositionDeltas->Num() > 0)
 	{
-		for (uint32 VertexIndex = 0; VertexIndex < VertexCount; ++VertexIndex)
+		if (IndexMap.Num() > 0)
 		{
-			const uint32 MappedVertexIndex = IndexMap[VertexIndex];
-			const FVector3f& Position = *reinterpret_cast<const FVector3f*>(SourceData + Stride * MappedVertexIndex);
-			Positions[VertexIndex] = FGLTFCoreUtilities::ConvertPosition(Position, Builder.ExportOptions->ExportUniformScale);
+			for (uint32 VertexIndex = 0; VertexIndex < VertexCount; ++VertexIndex)
+			{
+				const uint32 MappedVertexIndex = IndexMap[VertexIndex];
+				FVector3f Position = *reinterpret_cast<const FVector3f*>(SourceData + Stride * MappedVertexIndex);
+
+				if (TargetPositionDeltas->Contains(MappedVertexIndex))
+				{
+					Position += TargetPositionDeltas->FindChecked(MappedVertexIndex);
+				}
+				
+				Positions[VertexIndex] = FGLTFCoreUtilities::ConvertPosition(Position, Builder.ExportOptions->ExportUniformScale);
+			}
+		}
+		else
+		{
+			for (uint32 VertexIndex = 0; VertexIndex < VertexCount; ++VertexIndex)
+			{
+				FVector3f Position = *reinterpret_cast<const FVector3f*>(SourceData + Stride * VertexIndex);
+
+				if (TargetPositionDeltas->Contains(VertexIndex))
+				{
+					Position += TargetPositionDeltas->FindChecked(VertexIndex);
+				}
+
+				Positions[VertexIndex] = FGLTFCoreUtilities::ConvertPosition(Position, Builder.ExportOptions->ExportUniformScale);
+			}
 		}
 	}
 	else
 	{
-		for (uint32 VertexIndex = 0; VertexIndex < VertexCount; ++VertexIndex)
+		if (IndexMap.Num() > 0)
 		{
-			const FVector3f& Position = *reinterpret_cast<const FVector3f*>(SourceData + Stride * VertexIndex);
-			Positions[VertexIndex] = FGLTFCoreUtilities::ConvertPosition(Position, Builder.ExportOptions->ExportUniformScale);
+			for (uint32 VertexIndex = 0; VertexIndex < VertexCount; ++VertexIndex)
+			{
+				const uint32 MappedVertexIndex = IndexMap[VertexIndex];
+				const FVector3f& Position = *reinterpret_cast<const FVector3f*>(SourceData + Stride * MappedVertexIndex);
+				Positions[VertexIndex] = FGLTFCoreUtilities::ConvertPosition(Position, Builder.ExportOptions->ExportUniformScale);
+			}
+		}
+		else
+		{
+			for (uint32 VertexIndex = 0; VertexIndex < VertexCount; ++VertexIndex)
+			{
+				const FVector3f& Position = *reinterpret_cast<const FVector3f*>(SourceData + Stride * VertexIndex);
+				Positions[VertexIndex] = FGLTFCoreUtilities::ConvertPosition(Position, Builder.ExportOptions->ExportUniformScale);
+			}
 		}
 	}
-
-
+	
 	FGLTFJsonAccessor* JsonAccessor = Builder.AddAccessor();
 	JsonAccessor->BufferView = Builder.AddBufferView(Positions, EGLTFJsonBufferTarget::ArrayBuffer);
 	JsonAccessor->ComponentType = EGLTFJsonComponentType::Float;
@@ -218,7 +262,7 @@ FGLTFJsonAccessor* FGLTFColorBufferConverter::Convert(const FGLTFMeshSection* Me
 	return JsonAccessor;
 }
 
-FGLTFJsonAccessor* FGLTFNormalBufferConverter::Convert(const FGLTFMeshSection* MeshSection, const FStaticMeshVertexBuffer* VertexBuffer)
+FGLTFJsonAccessor* FGLTFNormalBufferConverter::Convert(const FGLTFMeshSection* MeshSection, const FStaticMeshVertexBuffer* VertexBuffer, const TMap<uint32, FVector3f>* TargetNormalDeltas)
 {
 	if (VertexBuffer == nullptr || VertexBuffer->GetNumVertices() == 0)
 	{
@@ -247,13 +291,13 @@ FGLTFJsonAccessor* FGLTFNormalBufferConverter::Convert(const FGLTFMeshSection* M
 		if (bHighPrecision)
 		{
 			ComponentType = EGLTFJsonComponentType::Int16;
-			BufferView = ConvertBufferView<FGLTFInt16Vector4, FPackedRGBA16N>(MeshSection, SourceData);
+			BufferView = ConvertBufferView<FGLTFInt16Vector4, FPackedRGBA16N>(MeshSection, SourceData, TargetNormalDeltas);
 			BufferView->ByteStride = sizeof(FGLTFInt16Vector4);
 		}
 		else
 		{
 			ComponentType = EGLTFJsonComponentType::Int8;
-			BufferView = ConvertBufferView<FGLTFInt8Vector4, FPackedNormal>(MeshSection, SourceData);
+			BufferView = ConvertBufferView<FGLTFInt8Vector4, FPackedNormal>(MeshSection, SourceData, TargetNormalDeltas);
 			BufferView->ByteStride = sizeof(FGLTFInt8Vector4);
 		}
 	}
@@ -261,8 +305,8 @@ FGLTFJsonAccessor* FGLTFNormalBufferConverter::Convert(const FGLTFMeshSection* M
 	{
 		ComponentType = EGLTFJsonComponentType::Float;
 		BufferView = bHighPrecision
-			? ConvertBufferView<FGLTFVector3, FPackedRGBA16N>(MeshSection, SourceData)
-			: ConvertBufferView<FGLTFVector3, FPackedNormal>(MeshSection, SourceData);
+			? ConvertBufferView<FGLTFVector3, FPackedRGBA16N>(MeshSection, SourceData, TargetNormalDeltas)
+			: ConvertBufferView<FGLTFVector3, FPackedNormal>(MeshSection, SourceData, TargetNormalDeltas);
 	}
 
 	FGLTFJsonAccessor* JsonAccessor = Builder.AddAccessor();
@@ -276,7 +320,7 @@ FGLTFJsonAccessor* FGLTFNormalBufferConverter::Convert(const FGLTFMeshSection* M
 }
 
 template <typename DestinationType, typename SourceType>
-FGLTFJsonBufferView* FGLTFNormalBufferConverter::ConvertBufferView(const FGLTFMeshSection* MeshSection, const void* SourceData)
+FGLTFJsonBufferView* FGLTFNormalBufferConverter::ConvertBufferView(const FGLTFMeshSection* MeshSection, const void* SourceData, const TMap<uint32, FVector3f>* TargetNormalDeltas)
 {
 	const TArray<uint32>& IndexMap = MeshSection->IndexMap;
 	const uint32 VertexCount = IndexMap.Num();
@@ -287,11 +331,33 @@ FGLTFJsonBufferView* FGLTFNormalBufferConverter::ConvertBufferView(const FGLTFMe
 	TArray<DestinationType> Normals;
 	Normals.AddUninitialized(VertexCount);
 
-	for (uint32 VertexIndex = 0; VertexIndex < VertexCount; ++VertexIndex)
+	if (TargetNormalDeltas && TargetNormalDeltas->Num())
 	{
-		const uint32 MappedVertexIndex = IndexMap[VertexIndex];
-		const TangentDatumType& TangentDatum = TangentData[MappedVertexIndex];
-		Normals[VertexIndex] = TGLTFVertexNormalUtilities<DestinationType, SourceType>::Convert(TangentDatum.TangentZ);
+		for (uint32 VertexIndex = 0; VertexIndex < VertexCount; ++VertexIndex)
+		{
+			const uint32 MappedVertexIndex = IndexMap[VertexIndex];
+			const TangentDatumType& TangentDatum = TangentData[MappedVertexIndex];
+			
+
+			const FVector3f* TargetNormalDelta = TargetNormalDeltas->Find(MappedVertexIndex);
+			if (TargetNormalDelta)
+			{
+				Normals[VertexIndex] = TGLTFVertexNormalUtilities<DestinationType, SourceType>::Convert(TangentDatum.TangentZ, *TargetNormalDelta);
+			}
+			else
+			{
+				Normals[VertexIndex] = TGLTFVertexNormalUtilities<DestinationType, SourceType>::Convert(TangentDatum.TangentZ);
+			}
+		}
+	}
+	else
+	{
+		for (uint32 VertexIndex = 0; VertexIndex < VertexCount; ++VertexIndex)
+		{
+			const uint32 MappedVertexIndex = IndexMap[VertexIndex];
+			const TangentDatumType& TangentDatum = TangentData[MappedVertexIndex];
+			Normals[VertexIndex] = TGLTFVertexNormalUtilities<DestinationType, SourceType>::Convert(TangentDatum.TangentZ);
+		}
 	}
 
 	return Builder.AddBufferView(Normals, EGLTFJsonBufferTarget::ArrayBuffer);
diff --git a/Engine/Plugins/Enterprise/GLTFExporter/Source/GLTFExporter/Private/Json/GLTFJsonMesh.cpp b/Engine/Plugins/Enterprise/GLTFExporter/Source/GLTFExporter/Private/Json/GLTFJsonMesh.cpp
--- a/Engine/Plugins/Enterprise/GLTFExporter/Source/GLTFExporter/Private/Json/GLTFJsonMesh.cpp
+++ b/Engine/Plugins/Enterprise/GLTFExporter/Source/GLTFExporter/Private/Json/GLTFJsonMesh.cpp
@@ -4,6 +4,19 @@
 #include "Json/GLTFJsonAccessor.h"
 #include "Json/GLTFJsonMaterial.h"
 
+void FGLTFJsonTarget::WriteObject(IGLTFJsonWriter& Writer) const
+{
+	if (Position != nullptr) Writer.Write(TEXT("POSITION"), Position);
+	if (Normal != nullptr) Writer.Write(TEXT("NORMAL"), Normal);
+}
+
+bool FGLTFJsonTarget::HasValue() const
+{
+	return (Position && Position->Count)
+		|| (Normal && Normal->Count);
+}
+
+
 void FGLTFJsonAttributes::WriteObject(IGLTFJsonWriter& Writer) const
 {
 	if (Position != nullptr) Writer.Write(TEXT("POSITION"), Position);
@@ -78,6 +91,11 @@ void FGLTFJsonPrimitive::WriteObject(IGLTFJsonWriter& Writer) const
 
 		Writer.EndExtensions();
 	}
+
+	if (Targets.Num() > 0)
+	{
+		Writer.Write(TEXT("targets"), Targets);
+	}
 }
 
 bool FGLTFJsonPrimitive::HasValue() const
@@ -102,6 +120,18 @@ void FGLTFJsonMesh::WriteObject(IGLTFJsonWriter& Writer) const
 		}
 	}
 	Writer.EndArray();
+
+	if (TargetNames.Num())
+	{
+		Writer.StartObject(TEXT("extras"));
+		Writer.StartArray(TEXT("targetNames"));
+		for (const FString& TargetName : TargetNames)
+		{
+			Writer.Write(TargetName);
+		}
+		Writer.EndArray();
+		Writer.EndObject();
+	}
 }
 
 bool FGLTFJsonMesh::HasValue() const
diff --git a/Engine/Plugins/Enterprise/GLTFExporter/Source/GLTFExporter/Private/Options/GLTFExportOptions.cpp b/Engine/Plugins/Enterprise/GLTFExporter/Source/GLTFExporter/Private/Options/GLTFExportOptions.cpp
--- a/Engine/Plugins/Enterprise/GLTFExporter/Source/GLTFExporter/Private/Options/GLTFExportOptions.cpp
+++ b/Engine/Plugins/Enterprise/GLTFExporter/Source/GLTFExporter/Private/Options/GLTFExportOptions.cpp
@@ -34,6 +34,7 @@ void UGLTFExportOptions::ResetToDefault()
 	bExportVertexSkinWeights = true;
 	bMakeSkinnedMeshesRoot = true;
 	bUseMeshQuantization = false;
+	bExportMorphTargets = true;
 	bExportLevelSequences = true;
 	bExportAnimationSequences = true;
 	TextureImageFormat = EGLTFTextureImageFormat::PNG;
diff --git a/Engine/Plugins/Enterprise/GLTFExporter/Source/GLTFExporter/Private/Tasks/GLTFDelayedMeshTasks.cpp b/Engine/Plugins/Enterprise/GLTFExporter/Source/GLTFExporter/Private/Tasks/GLTFDelayedMeshTasks.cpp
--- a/Engine/Plugins/Enterprise/GLTFExporter/Source/GLTFExporter/Private/Tasks/GLTFDelayedMeshTasks.cpp
+++ b/Engine/Plugins/Enterprise/GLTFExporter/Source/GLTFExporter/Private/Tasks/GLTFDelayedMeshTasks.cpp
@@ -1,6 +1,7 @@
 // Copyright Epic Games, Inc. All Rights Reserved.
 
 #include "Tasks/GLTFDelayedMeshTasks.h"
+#include "Algo/AllOf.h"
 #include "Converters/GLTFMeshUtilities.h"
 #include "Converters/GLTFBufferAdapter.h"
 #include "Builders/GLTFConvertBuilder.h"
@@ -196,11 +197,13 @@ void FGLTFDelayedStaticAndSplineMeshTask::ProcessMeshDescription(const TArray<FS
 	FMeshDescriptionParser MeshDescriptionParser(MeshDescription, MaterialSlots);
 
 	constexpr bool bExportVertexSkinWeights_False = false;
+	constexpr bool bExportMorphTargets_False = false;
 	constexpr int32 SkeletonInfluenceCountPerGroup_4 = 4;
-	FExportConfigs ExportConfigs(bExportVertexSkinWeights_False, Builder.ExportOptions->bExportVertexColors, SplineMeshComponent, SkeletonInfluenceCountPerGroup_4);
+	FExportConfigs ExportConfigs(bExportVertexSkinWeights_False, Builder.ExportOptions->bExportVertexColors, SplineMeshComponent, SkeletonInfluenceCountPerGroup_4, bExportMorphTargets_False);
 
 	TArray<FMeshPrimitiveDescription> MeshPrimitiveDescriptions;
-	MeshDescriptionParser.Parse(MeshPrimitiveDescriptions, ExportConfigs);
+	TArray<FString> TargetNames; //Dummy variable
+	MeshDescriptionParser.Parse(MeshPrimitiveDescriptions, TargetNames, ExportConfigs);
 	
 	if (MeshPrimitiveDescriptions.Num() != JsonMesh->Primitives.Num())
 	{
@@ -418,11 +421,11 @@ void FGLTFDelayedSkeletalMeshTask::ProcessSourceModel(const TArray<FSkeletalMate
 	
 	constexpr bool bExportVertexSkinWeights_True = true;
 	constexpr int32 SkeletonInfluenceCountPerGroup_4 = 4;
-	FExportConfigs ExportConfigs(bExportVertexSkinWeights_True, Builder.ExportOptions->bExportVertexColors, nullptr, SkeletonInfluenceCountPerGroup_4);
+	FExportConfigs ExportConfigs(bExportVertexSkinWeights_True, Builder.ExportOptions->bExportVertexColors, nullptr, SkeletonInfluenceCountPerGroup_4, Builder.ExportOptions->bExportMorphTargets);
 	
 
 	TArray<FMeshPrimitiveDescription> MeshPrimitiveDescriptions;
-	MeshDescriptionParser.Parse(MeshPrimitiveDescriptions, ExportConfigs);
+	MeshDescriptionParser.Parse(MeshPrimitiveDescriptions, JsonMesh->TargetNames, ExportConfigs);
 
 	if (MeshPrimitiveDescriptions.Num() != JsonMesh->Primitives.Num())
 	{
@@ -452,6 +455,16 @@ void FGLTFDelayedSkeletalMeshTask::ProcessSourceModel(const TArray<FSkeletalMate
 			JsonPrimitive.Attributes.TexCoords[UVIndex] = Builder.AddUniqueUVAccessor(MeshPrimitiveDescription.UVs[UVIndex]);
 		}
 
+		for (FPrimitiveTargetDescription& PrimitiveTargetDescription : MeshPrimitiveDescription.Targets)
+		{
+			if (PrimitiveTargetDescription.bHasPositions || PrimitiveTargetDescription.bHasNormals)
+			{
+				FGLTFJsonTarget& JSONMorphTarget = JsonPrimitive.Targets.AddDefaulted_GetRef();
+				if (PrimitiveTargetDescription.bHasPositions) JSONMorphTarget.Position = Builder.AddUniquePositionAccessor(PrimitiveTargetDescription.Positions);
+				if (PrimitiveTargetDescription.bHasNormals) JSONMorphTarget.Normal = Builder.AddUniqueNormalAccessor(PrimitiveTargetDescription.Normals);
+			}
+		}
+
 		if (Builder.ExportOptions->bExportVertexSkinWeights)
 		{
 			int32 GroupCount = MeshPrimitiveDescription.JointInfluences.Num();
@@ -519,6 +532,48 @@ void FGLTFDelayedSkeletalMeshTask::ProcessRenderData(const TArray<FSkeletalMater
 
 	const int32 MaterialCount = MaterialSlots.Num();
 
+	struct FTargetDeltas
+	{
+		TMap<uint32, FVector3f> Deltas;
+		bool bIsValid = false;
+
+		void Add(const uint32 VertexIndex, const FVector3f& Delta)
+		{
+			Deltas.Add(VertexIndex, Delta);
+		}
+
+		void GenerateIsValid()
+		{
+			bIsValid = Deltas.Num() && !Algo::AllOf(Deltas, [](const TPair<uint32, FVector3f>& Element) {
+				return Element.Value.IsZero();
+			});
+		}
+	};
+	
+	TArray<FTargetDeltas> MorphTargetsPositions;
+	TArray<FTargetDeltas> MorphTargetsNormals;
+
+	if (Builder.ExportOptions->bExportMorphTargets)
+	{
+		const TArray<TObjectPtr<UMorphTarget>>& MorphTargetObjectPtrs = SkeletalMesh->GetMorphTargets();
+		for (const TObjectPtr<UMorphTarget>& MorphTargetObjectPtr : MorphTargetObjectPtrs)
+		{
+			const UMorphTarget* MorphTarget = MorphTargetObjectPtr;
+
+			FTargetDeltas& PositionsDeltas = MorphTargetsPositions.AddDefaulted_GetRef();
+			FTargetDeltas& NormalsDeltas = MorphTargetsNormals.AddDefaulted_GetRef();
+
+			for (FMorphTargetDeltaIterator DeltaIt = MorphTarget->GetDeltaIteratorForLOD(LODIndex); !DeltaIt.AtEnd(); ++DeltaIt)
+			{
+				PositionsDeltas.Add(DeltaIt->SourceIdx, DeltaIt->PositionDelta);
+				NormalsDeltas.Add(DeltaIt->SourceIdx, DeltaIt->TangentZDelta);
+			}
+
+			PositionsDeltas.GenerateIsValid();
+			NormalsDeltas.GenerateIsValid();
+		}
+	}
+
 	for (int32 MaterialIndex = 0; MaterialIndex < MaterialCount; ++MaterialIndex)
 	{
 		const FGLTFIndexArray SectionIndices = FGLTFMeshUtilities::GetSectionIndices(RenderData, MaterialIndex);
@@ -538,6 +593,24 @@ void FGLTFDelayedSkeletalMeshTask::ProcessRenderData(const TArray<FSkeletalMater
 				));
 		}
 
+		if (MorphTargetsPositions.Num() > 0)
+		{
+			JsonPrimitive.Targets.Reserve(MorphTargetsPositions.Num());
+
+			for (size_t TargetCounter = 0; TargetCounter < MorphTargetsPositions.Num(); TargetCounter++)
+			{
+				FTargetDeltas& PositionsDeltas = MorphTargetsPositions[TargetCounter];
+				FTargetDeltas& NormalsDeltas = MorphTargetsNormals[TargetCounter];
+
+				if (PositionsDeltas.bIsValid || NormalsDeltas.bIsValid)
+				{
+					FGLTFJsonTarget& JSONMorphTarget = JsonPrimitive.Targets.AddDefaulted_GetRef();
+					if (PositionsDeltas.bIsValid) JSONMorphTarget.Position = Builder.AddUniquePositionAccessor(ConvertedSection, &PositionBuffer, &PositionsDeltas.Deltas);
+					if (NormalsDeltas.bIsValid) JSONMorphTarget.Normal = Builder.AddUniqueNormalAccessor(ConvertedSection, &VertexBuffer, &NormalsDeltas.Deltas);
+				}
+			}
+		}
+
 		if (ColorBuffer != nullptr)
 		{
 			JsonPrimitive.Attributes.Color0 = Builder.AddUniqueColorAccessor(ConvertedSection, ColorBuffer);
diff --git a/Engine/Plugins/Enterprise/GLTFExporter/Source/GLTFExporter/Private/Utilities/GLTFCoreUtilities.cpp b/Engine/Plugins/Enterprise/GLTFExporter/Source/GLTFExporter/Private/Utilities/GLTFCoreUtilities.cpp
--- a/Engine/Plugins/Enterprise/GLTFExporter/Source/GLTFExporter/Private/Utilities/GLTFCoreUtilities.cpp
+++ b/Engine/Plugins/Enterprise/GLTFExporter/Source/GLTFExporter/Private/Utilities/GLTFCoreUtilities.cpp
@@ -43,6 +43,25 @@ FGLTFInt16Vector4 FGLTFCoreUtilities::ConvertNormal(const FPackedRGBA16N& Normal
 	return { SafeNormal.X, SafeNormal.Z, SafeNormal.Y, 0 };
 }
 
+FGLTFVector3 FGLTFCoreUtilities::ConvertNormal(const FVector3f& Normal, const FVector3f& Offset)
+{
+	const FVector3f SafeNormal = (Normal + Offset).GetSafeNormal();
+	return ConvertVector(SafeNormal);
+}
+
+FGLTFInt8Vector4 FGLTFCoreUtilities::ConvertNormal(const FPackedNormal& Normal, const FVector3f& Offset)
+{
+	const FPackedNormal SafeNormal = (Normal.ToFVector3f() + Offset).GetSafeNormal();
+	return { SafeNormal.Vector.X, SafeNormal.Vector.Z, SafeNormal.Vector.Y, 0 };
+}
+
+FGLTFInt16Vector4 FGLTFCoreUtilities::ConvertNormal(const FPackedRGBA16N& Normal, const FVector3f& Offset)
+{
+	const FPackedRGBA16N SafeNormal = (Normal.ToFVector3f() + Offset).GetSafeNormal();
+	return { SafeNormal.X, SafeNormal.Z, SafeNormal.Y, 0 };
+}
+
+
 FGLTFVector4 FGLTFCoreUtilities::ConvertTangent(const FVector3f& Tangent, const FVector4f& Normal)
 {
 	// Unreal keeps the binormal sign in the normal's w-component.
diff --git a/Engine/Plugins/Enterprise/GLTFExporter/Source/GLTFExporter/Private/Utilities/MeshParsingUtilities.cpp b/Engine/Plugins/Enterprise/GLTFExporter/Source/GLTFExporter/Private/Utilities/MeshParsingUtilities.cpp
--- a/Engine/Plugins/Enterprise/GLTFExporter/Source/GLTFExporter/Private/Utilities/MeshParsingUtilities.cpp
+++ b/Engine/Plugins/Enterprise/GLTFExporter/Source/GLTFExporter/Private/Utilities/MeshParsingUtilities.cpp
@@ -53,6 +53,20 @@ namespace UE::MeshParser
 			}
 			return false;
 		}
+
+		struct FVertexIdTracker
+		{
+			int32 ExportedVertexId;
+			FVertexID VertexId; //original
+			FVertexInstanceID VertexInstanceId; //original
+
+			FVertexIdTracker(const int32& InExportedVertexId, const FVertexID& InVertexId, const FVertexInstanceID& InVertexInstanceId)
+				: ExportedVertexId(InExportedVertexId)
+				, VertexId(InVertexId)
+				, VertexInstanceId(InVertexInstanceId)
+			{
+			}
+		};
 	}
 
 	//
@@ -77,7 +91,7 @@ namespace UE::MeshParser
 	}
 
 	//Except for Joint related containers
-	void FMeshPrimitiveDescription::PrepareContainers(const int32& IndexCount, const int32& AttributesCount, const int32& UVCount, bool bPrepareVertexColors)
+	void FMeshPrimitiveDescription::PrepareContainers(const int32& IndexCount, const int32& AttributesCount, const int32& UVCount, bool bPrepareVertexColors, const int32& TargetCount)
 	{
 		Indices.Reserve(IndexCount);
 		Positions.SetNumZeroed(AttributesCount);
@@ -89,6 +103,12 @@ namespace UE::MeshParser
 		{
 			UVs[UVIndex].SetNumZeroed(AttributesCount);
 		}
+		for (size_t TargetCounter = 0; TargetCounter < TargetCount; TargetCounter++)
+		{
+			FPrimitiveTargetDescription& TargetDescription = Targets.AddDefaulted_GetRef();
+			TargetDescription.Positions.SetNumZeroed(AttributesCount);
+			TargetDescription.Normals.SetNumZeroed(AttributesCount);
+		}
 	}
 
 	void FMeshPrimitiveDescription::PrepareJointContainers(const uint32& JointGroupCount, const int32& AttributesCount)
@@ -105,11 +125,13 @@ namespace UE::MeshParser
 	FExportConfigs::FExportConfigs(bool bInExportVertexSkinWeights,
 		bool bInExportVertexColors,
 		const USplineMeshComponent* InSplineMeshComponent,
-		int32 InSkeletonInfluenceCountPerGroup)
+		int32 InSkeletonInfluenceCountPerGroup,
+		bool bInExportMorphTargets)
 		: bExportVertexSkinWeights(bInExportVertexSkinWeights)
 		, bExportVertexColors(bInExportVertexColors)
 		, SplineMeshComponent(InSplineMeshComponent)
 		, SkeletonInfluenceCountPerGroup(InSkeletonInfluenceCountPerGroup)
+		, bExportMorphTargets(bInExportMorphTargets)
 	{
 	}
 
@@ -135,6 +157,58 @@ namespace UE::MeshParser
 		VertexInstanceIdToVertexId = MeshAttributes.GetVertexInstanceVertexIndices().GetRawArray();
 		PolygonGroupMaterialSlotNames = MeshAttributes.GetPolygonGroupMaterialSlotNames();
 
+		//MorphTargets:
+		if constexpr (std::is_same_v<T, FSkeletalMaterial>)
+		{
+			FSkeletalMeshConstAttributes SkeletalMeshAttributes(*MeshDescription);
+			TArray<FName> MorphTargetNames = SkeletalMeshAttributes.GetMorphTargetNames();
+
+			MorphTargetNames.Sort([](const FName& Left, const FName Right)
+				{
+					return Left.ToString() < Right.ToString();
+				});
+
+			TargetVertexPositionDeltas.Reserve(MorphTargetNames.Num());
+			TargetVertexInstanceNormalDeltas.Reserve(MorphTargetNames.Num());
+			TargetNames.Reserve(MorphTargetNames.Num());
+
+			for (const FName& MorphTargetName : MorphTargetNames)
+			{
+				if (SkeletalMeshAttributes.HasMorphTargetPositionsAttribute(MorphTargetName) || SkeletalMeshAttributes.HasMorphTargetNormalsAttribute(MorphTargetName))
+				{
+					TargetNames.Add(MorphTargetName.ToString());
+				}
+				else
+				{
+					continue;
+				}
+
+				if (SkeletalMeshAttributes.HasMorphTargetPositionsAttribute(MorphTargetName))
+				{
+					TargetVertexPositionDeltas.Add(SkeletalMeshAttributes.GetVertexMorphPositionDelta(MorphTargetName));
+				}
+				else
+				{
+					TargetVertexPositionDeltas.AddDefaulted();
+				}
+
+				if (SkeletalMeshAttributes.HasMorphTargetNormalsAttribute(MorphTargetName))
+				{
+					TargetVertexInstanceNormalDeltas.Add(SkeletalMeshAttributes.GetVertexInstanceMorphNormalDelta(MorphTargetName));
+				}
+				else
+				{
+					TargetVertexInstanceNormalDeltas.AddDefaulted();
+				}
+			}
+
+			NumOfTargets = TargetNames.Num();
+		}
+		else
+		{
+			NumOfTargets = 0;
+		}
+		
 		//Mesh level details that are used for the parsing
 		MeshDetails.UVCount = VertexInstanceUVs.GetNumChannels();
 		MeshDetails.NumberOfPrimitives = MeshDescription->PolygonGroups().Num();
@@ -142,10 +216,12 @@ namespace UE::MeshParser
 	}
 
 	template <class T>
-	void FMeshDescriptionParser<T>::Parse(TArray<FMeshPrimitiveDescription>& OutMeshPrimitiveDescriptions, const FExportConfigs& ExportConfigs)
+	void FMeshDescriptionParser<T>::Parse(TArray<FMeshPrimitiveDescription>& OutMeshPrimitiveDescriptions, TArray<FString>& OutTargetNames, const FExportConfigs& ExportConfigs)
 	{
 		using namespace Private;
 
+		OutTargetNames = TargetNames;
+
 		//Iterate through all PolygonsGroups:
 		for (const FPolygonGroupID& PolygonGroupID : MeshDescription->PolygonGroups().GetElementIDs())
 		{
@@ -239,26 +315,33 @@ namespace UE::MeshParser
 			TArray<bool> Normals;
 			TArray<bool> Tangents;
 			TArray<TArray<bool>> UVs;
+			TArray<TArray<bool>> TargetNormals;
 
-			FVertexAttributesSetTracker(int32 VertexCount, int32 UVCount)
+			FVertexAttributesSetTracker(int32 VertexCount, int32 UVCount, int32 TargetCount)
 			{
 				VertexColors.SetNumZeroed(VertexCount);
 				Normals.SetNumZeroed(VertexCount);
 				Tangents.SetNumZeroed(VertexCount);
 				for (size_t UVIndex = 0; UVIndex < UVCount; UVIndex++)
 				{
-					UVs.Add(TArray<bool>());
-					UVs[UVIndex].SetNumZeroed(VertexCount);
+					UVs.AddDefaulted_GetRef().SetNumZeroed(VertexCount);
+				}
+				for (size_t TargetCounter = 0; TargetCounter < TargetCount; TargetCounter++)
+				{
+					TargetNormals.AddDefaulted_GetRef().SetNumZeroed(VertexCount);
 				}
 			}
 		};
-		FVertexAttributesSetTracker AttributesSetTracker(VertexCount, MeshDetails.UVCount);
+		FVertexAttributesSetTracker AttributesSetTracker(VertexCount, MeshDetails.UVCount, NumOfTargets);
 
 		const bool bCanExportVertexColors = MeshDetails.bHasVertexColors && ExportConfigs.bExportVertexColors;
 
 		//Set and Reserve the appropriate counts for the Attributes (and Indices)
 		//SetNum/Reserve the containers:
-		MeshPrimitiveDescription.PrepareContainers(VertexInstanceCount, VertexCount, MeshDetails.UVCount, bCanExportVertexColors);
+		MeshPrimitiveDescription.PrepareContainers(VertexInstanceCount, VertexCount, MeshDetails.UVCount, bCanExportVertexColors, NumOfTargets);
+
+		TArray<FVertexIdTracker> VertexIdTrackers; //Exported to Original(VertexId, VertexInstanceId)
+		VertexIdTrackers.Reserve(TriangleIDs.Num() * 3);
 
 		for (const FTriangleID& TriangleID : TriangleIDs)
 		{
@@ -269,6 +352,8 @@ namespace UE::MeshParser
 				int32* ExportedVertexIDPtr = OriginalToExportedVertexIdMap.Find(OriginalVertexID);
 				int32 ExportedVertexID = ExportedVertexIDPtr ? *ExportedVertexIDPtr : OriginalToExportedVertexIdMap.Num();
 
+				VertexIdTrackers.Add(FVertexIdTracker(ExportedVertexID, OriginalVertexID, VertexInstanceID));
+
 				MeshPrimitiveDescription.Indices.Add(ExportedVertexID);
 
 				//Positions are Vertex based so they can be directly set, no need to check them:
@@ -300,6 +385,69 @@ namespace UE::MeshParser
 			}
 		}
 
+		if (ExportConfigs.bExportMorphTargets)
+		{
+			//TargetNames and VertexPositionDeltas will have the exact same amount of Targets
+			for (size_t TargetCounter = 0; TargetCounter < NumOfTargets; TargetCounter++)
+			{
+				MeshPrimitiveDescription.Targets[TargetCounter].bHasPositions = TargetVertexPositionDeltas[TargetCounter].IsValid();
+				MeshPrimitiveDescription.Targets[TargetCounter].bHasNormals = TargetVertexInstanceNormalDeltas[TargetCounter].IsValid();
+
+				if (MeshPrimitiveDescription.Targets[TargetCounter].bHasPositions)
+				{
+					if (MeshPrimitiveDescription.Targets[TargetCounter].bHasNormals)
+					{
+						//Has both
+						for (const FVertexIdTracker& VertexIdTracker : VertexIdTrackers)
+						{
+							int32 ExportedVertexID = VertexIdTracker.ExportedVertexId;
+							FVertexID OriginalVertexID = VertexIdTracker.VertexId;
+							FVertexInstanceID VertexInstanceID = VertexIdTracker.VertexInstanceId;
+
+							MeshPrimitiveDescription.Targets[TargetCounter].Positions[ExportedVertexID] = MeshPrimitiveDescription.Positions[ExportedVertexID] + TargetVertexPositionDeltas[TargetCounter][OriginalVertexID];
+
+							if (!CheckCompareSetValue(AttributesSetTracker.TargetNormals[TargetCounter][ExportedVertexID],
+								MeshPrimitiveDescription.Targets[TargetCounter].Normals[ExportedVertexID],
+								VertexInstanceNormals[VertexInstanceID] + TargetVertexInstanceNormalDeltas[TargetCounter][OriginalVertexID]))
+							{
+								return false;
+							}
+						}
+					}
+					else
+					{
+						//Has only positions
+						for (const FVertexIdTracker& VertexIdTracker : VertexIdTrackers)
+						{
+							int32 ExportedVertexID = VertexIdTracker.ExportedVertexId;
+							FVertexID OriginalVertexID = VertexIdTracker.VertexId;
+							FVertexInstanceID VertexInstanceID = VertexIdTracker.VertexInstanceId;
+
+							MeshPrimitiveDescription.Targets[TargetCounter].Positions[ExportedVertexID] = MeshPrimitiveDescription.Positions[ExportedVertexID] + TargetVertexPositionDeltas[TargetCounter][OriginalVertexID];
+						}
+					}
+				}
+				else if (MeshPrimitiveDescription.Targets[TargetCounter].bHasNormals)
+				{
+					//Has only normals
+					//Has both
+					for (const FVertexIdTracker& VertexIdTracker : VertexIdTrackers)
+					{
+						int32 ExportedVertexID = VertexIdTracker.ExportedVertexId;
+						FVertexID OriginalVertexID = VertexIdTracker.VertexId;
+						FVertexInstanceID VertexInstanceID = VertexIdTracker.VertexInstanceId;
+
+						if (!CheckCompareSetValue(AttributesSetTracker.TargetNormals[TargetCounter][ExportedVertexID],
+							MeshPrimitiveDescription.Targets[TargetCounter].Normals[ExportedVertexID],
+							VertexInstanceNormals[VertexInstanceID] + TargetVertexInstanceNormalDeltas[TargetCounter][OriginalVertexID]))
+						{
+							return false;
+						}
+					}
+				}
+			}
+		}
+
 		//Set VertexColors if MeshDescription has any:
 		if (bCanExportVertexColors)
 		{
@@ -412,7 +560,10 @@ namespace UE::MeshParser
 		const bool bCanExportVertexColors = MeshDetails.bHasVertexColors && ExportConfigs.bExportVertexColors;
 
 		//SetNum/Reserve the containers:
-		MeshPrimitiveDescription.PrepareContainers(VertexInstanceCount, VertexInstanceCount, MeshDetails.UVCount, bCanExportVertexColors);
+		MeshPrimitiveDescription.PrepareContainers(VertexInstanceCount, VertexInstanceCount, MeshDetails.UVCount, bCanExportVertexColors, TargetNames.Num());
+
+		TArray<FVertexIdTracker> VertexIdTrackers; //Exported to Original(VertexId, VertexInstanceId), primarily used for Targets.
+		VertexIdTrackers.Reserve(TriangleIDs.Num() * 3);
 
 		//Based on Triangles acquire the VertexInstanceID, based on which acquire the VertexInstance Attributes:
 		for (const FTriangleID& TriangleID : TriangleIDs)
@@ -424,6 +575,8 @@ namespace UE::MeshParser
 				int32* ExportedVertexIDPtr = OriginalToExportedVertexIdMap.Find(VertexInstanceID);
 				int32 ExportedVertexID = ExportedVertexIDPtr ? *ExportedVertexIDPtr : OriginalToExportedVertexIdMap.Num();
 
+				VertexIdTrackers.Add(FVertexIdTracker(ExportedVertexID, OriginalVertexID, VertexInstanceID));
+
 				MeshPrimitiveDescription.Indices.Add(ExportedVertexID);
 
 				//Positions and Indicies are Vertex based so they can be directly set, no need to check them:
@@ -441,6 +594,58 @@ namespace UE::MeshParser
 						MeshPrimitiveDescription.UVs[UVIndex][ExportedVertexID] = VertexInstanceUVs.Get(VertexInstanceID, UVIndex);
 					}
 				}
+
+			}
+		}
+
+		if (ExportConfigs.bExportMorphTargets)
+		{
+			//TargetNames and VertexPositionDeltas will have the exact same amount of Targets
+			for (size_t TargetCounter = 0; TargetCounter < NumOfTargets; TargetCounter++)
+			{
+				MeshPrimitiveDescription.Targets[TargetCounter].bHasPositions = TargetVertexPositionDeltas[TargetCounter].IsValid();
+				MeshPrimitiveDescription.Targets[TargetCounter].bHasNormals = TargetVertexInstanceNormalDeltas[TargetCounter].IsValid();
+
+				if (MeshPrimitiveDescription.Targets[TargetCounter].bHasPositions)
+				{
+					if (MeshPrimitiveDescription.Targets[TargetCounter].bHasNormals)
+					{
+						//Has both
+						for (const FVertexIdTracker& VertexIdTracker : VertexIdTrackers)
+						{
+							int32 ExportedVertexID = VertexIdTracker.ExportedVertexId;
+							FVertexID OriginalVertexID = VertexIdTracker.VertexId;
+							FVertexInstanceID VertexInstanceID = VertexIdTracker.VertexInstanceId;
+
+							MeshPrimitiveDescription.Targets[TargetCounter].Positions[ExportedVertexID] = VertexPositions[OriginalVertexID] + (TargetVertexPositionDeltas[TargetCounter].IsValid() ? TargetVertexPositionDeltas[TargetCounter][OriginalVertexID] : FVector3f());
+							MeshPrimitiveDescription.Targets[TargetCounter].Normals[ExportedVertexID] = VertexInstanceNormals[VertexInstanceID] + (TargetVertexInstanceNormalDeltas[TargetCounter].IsValid() ? TargetVertexInstanceNormalDeltas[TargetCounter][VertexInstanceID] : FVector3f());
+						}
+					}
+					else
+					{
+						//Has only positions
+						for (const FVertexIdTracker& VertexIdTracker : VertexIdTrackers)
+						{
+							int32 ExportedVertexID = VertexIdTracker.ExportedVertexId;
+							FVertexID OriginalVertexID = VertexIdTracker.VertexId;
+							FVertexInstanceID VertexInstanceID = VertexIdTracker.VertexInstanceId;
+
+							MeshPrimitiveDescription.Targets[TargetCounter].Positions[ExportedVertexID] = VertexPositions[OriginalVertexID] + (TargetVertexPositionDeltas[TargetCounter].IsValid() ? TargetVertexPositionDeltas[TargetCounter][OriginalVertexID] : FVector3f());
+						}
+					}
+				}
+				else if (MeshPrimitiveDescription.Targets[TargetCounter].bHasNormals)
+				{
+					//Has only normals
+					for (const FVertexIdTracker& VertexIdTracker : VertexIdTrackers)
+					{
+						int32 ExportedVertexID = VertexIdTracker.ExportedVertexId;
+						FVertexID OriginalVertexID = VertexIdTracker.VertexId;
+						FVertexInstanceID VertexInstanceID = VertexIdTracker.VertexInstanceId;
+
+						MeshPrimitiveDescription.Targets[TargetCounter].Normals[ExportedVertexID] = VertexInstanceNormals[VertexInstanceID] + (TargetVertexInstanceNormalDeltas[TargetCounter].IsValid() ? TargetVertexInstanceNormalDeltas[TargetCounter][VertexInstanceID] : FVector3f());
+					}
+				}
 			}
 		}
 
diff --git a/Engine/Plugins/Enterprise/GLTFExporter/Source/GLTFExporter/Public/Builders/GLTFConvertBuilder.h b/Engine/Plugins/Enterprise/GLTFExporter/Source/GLTFExporter/Public/Builders/GLTFConvertBuilder.h
--- a/Engine/Plugins/Enterprise/GLTFExporter/Source/GLTFExporter/Public/Builders/GLTFConvertBuilder.h
+++ b/Engine/Plugins/Enterprise/GLTFExporter/Source/GLTFExporter/Public/Builders/GLTFConvertBuilder.h
@@ -36,11 +36,11 @@ public:
 	UE_API bool IsSelectedActor(const AActor* Object) const;
 	UE_API bool IsRootActor(const AActor* Actor) const;
 
-	UE_API FGLTFJsonAccessor* AddUniquePositionAccessor(const FGLTFMeshSection* MeshSection, const FPositionVertexBuffer* VertexBuffer);
+	UE_API FGLTFJsonAccessor* AddUniquePositionAccessor(const FGLTFMeshSection* MeshSection, const FPositionVertexBuffer* VertexBuffer, const TMap<uint32, FVector3f>* TargetPositionDeltas = nullptr);
 	UE_API FGLTFJsonAccessor* AddUniquePositionAccessor(const FGLTFPositionArray& VertexBuffer);
 	UE_API FGLTFJsonAccessor* AddUniqueColorAccessor(const FGLTFMeshSection* MeshSection, const FColorVertexBuffer* VertexBuffer);
 	UE_API FGLTFJsonAccessor* AddUniqueColorAccessor(const FGLTFColorArray& VertexColorBuffer);
-	UE_API FGLTFJsonAccessor* AddUniqueNormalAccessor(const FGLTFMeshSection* MeshSection, const FStaticMeshVertexBuffer* VertexBuffer);
+	UE_API FGLTFJsonAccessor* AddUniqueNormalAccessor(const FGLTFMeshSection* MeshSection, const FStaticMeshVertexBuffer* VertexBuffer, const TMap<uint32, FVector3f>* TargetNormalDeltas = nullptr);
 	UE_API FGLTFJsonAccessor* AddUniqueNormalAccessor(const FGLTFNormalArray& Normals);
 	UE_API FGLTFJsonAccessor* AddUniqueTangentAccessor(const FGLTFMeshSection* MeshSection, const FStaticMeshVertexBuffer* VertexBuffer);
 	UE_API FGLTFJsonAccessor* AddUniqueTangentAccessor(const FGLTFTangentArray& Tangents);
diff --git a/Engine/Plugins/Enterprise/GLTFExporter/Source/GLTFExporter/Public/Converters/GLTFAccessorConverters.h b/Engine/Plugins/Enterprise/GLTFExporter/Source/GLTFExporter/Public/Converters/GLTFAccessorConverters.h
--- a/Engine/Plugins/Enterprise/GLTFExporter/Source/GLTFExporter/Public/Converters/GLTFAccessorConverters.h
+++ b/Engine/Plugins/Enterprise/GLTFExporter/Source/GLTFExporter/Public/Converters/GLTFAccessorConverters.h
@@ -15,9 +15,9 @@ class FColorVertexBuffer;
 class FStaticMeshVertexBuffer;
 class FSkinWeightVertexBuffer;
 
-typedef TGLTFConverter<FGLTFJsonAccessor*, const FGLTFMeshSection*, const FPositionVertexBuffer*> IGLTFPositionBufferConverter;
+typedef TGLTFConverter<FGLTFJsonAccessor*, const FGLTFMeshSection*, const FPositionVertexBuffer*, const TMap<uint32, FVector3f>*> IGLTFPositionBufferConverter;
 typedef TGLTFConverter<FGLTFJsonAccessor*, const FGLTFMeshSection*, const FColorVertexBuffer*> IGLTFColorBufferConverter;
-typedef TGLTFConverter<FGLTFJsonAccessor*, const FGLTFMeshSection*, const FStaticMeshVertexBuffer*> IGLTFNormalBufferConverter;
+typedef TGLTFConverter<FGLTFJsonAccessor*, const FGLTFMeshSection*, const FStaticMeshVertexBuffer*, const TMap<uint32, FVector3f>*> IGLTFNormalBufferConverter;
 typedef TGLTFConverter<FGLTFJsonAccessor*, const FGLTFMeshSection*, const FStaticMeshVertexBuffer*> IGLTFTangentBufferConverter;
 typedef TGLTFConverter<FGLTFJsonAccessor*, const FGLTFMeshSection*, const FStaticMeshVertexBuffer*, uint32> IGLTFUVBufferConverter;
 typedef TGLTFConverter<FGLTFJsonAccessor*, const FGLTFMeshSection*, const FSkinWeightVertexBuffer*, uint32> IGLTFBoneIndexBufferConverter;
@@ -42,7 +42,7 @@ public:
 
 protected:
 
-	UE_API virtual FGLTFJsonAccessor* Convert(const FGLTFMeshSection* MeshSection, const FPositionVertexBuffer* VertexBuffer) override;
+	UE_API virtual FGLTFJsonAccessor* Convert(const FGLTFMeshSection* MeshSection, const FPositionVertexBuffer* VertexBuffer, const TMap<uint32, FVector3f>* TargetPositionDeltas) override;
 };
 
 class FGLTFColorBufferConverter : public FGLTFBuilderContext, public IGLTFColorBufferConverter
@@ -64,12 +64,12 @@ public:
 
 protected:
 
-	UE_API virtual FGLTFJsonAccessor* Convert(const FGLTFMeshSection* MeshSection, const FStaticMeshVertexBuffer* VertexBuffer) override;
+	UE_API virtual FGLTFJsonAccessor* Convert(const FGLTFMeshSection* MeshSection, const FStaticMeshVertexBuffer* VertexBuffer, const TMap<uint32, FVector3f>* TargetNormalDeltas) override;
 
 private:
 
 	template <typename DestinationType, typename SourceType>
-	FGLTFJsonBufferView* ConvertBufferView(const FGLTFMeshSection* MeshSection, const void* TangentData);
+	FGLTFJsonBufferView* ConvertBufferView(const FGLTFMeshSection* MeshSection, const void* TangentData, const TMap<uint32, FVector3f>* TargetNormalDeltas);
 };
 
 class FGLTFTangentBufferConverter : public FGLTFBuilderContext, public IGLTFTangentBufferConverter
diff --git a/Engine/Plugins/Enterprise/GLTFExporter/Source/GLTFExporter/Public/Json/GLTFJsonMesh.h b/Engine/Plugins/Enterprise/GLTFExporter/Source/GLTFExporter/Public/Json/GLTFJsonMesh.h
--- a/Engine/Plugins/Enterprise/GLTFExporter/Source/GLTFExporter/Public/Json/GLTFJsonMesh.h
+++ b/Engine/Plugins/Enterprise/GLTFExporter/Source/GLTFExporter/Public/Json/GLTFJsonMesh.h
@@ -7,6 +7,23 @@
 
 #define UE_API GLTFEXPORTER_API
 
+//For MorphTargets
+struct FGLTFJsonTarget : IGLTFJsonObject
+{
+	FGLTFJsonAccessor* Position;
+	FGLTFJsonAccessor* Normal;
+
+	FGLTFJsonTarget()
+		: Position(nullptr)
+		, Normal(nullptr)
+	{
+	}
+
+	UE_API virtual void WriteObject(IGLTFJsonWriter& Writer) const override;
+
+	UE_API bool HasValue() const;
+};
+
 struct FGLTFJsonAttributes : IGLTFJsonObject
 {
 	FGLTFJsonAccessor* Position;
@@ -40,6 +57,8 @@ struct FGLTFJsonPrimitive : IGLTFJsonObject
 
 	TArray<FGLTFJsonMaterialVariantMapping> MaterialVariantMappings;
 
+	TArray<FGLTFJsonTarget> Targets; //MorphTargets
+
 	FGLTFJsonPrimitive()
 		: Indices(nullptr)
 		, Material(nullptr)
@@ -58,6 +77,8 @@ struct FGLTFJsonMesh : IGLTFJsonIndexedObject
 
 	TArray<FGLTFJsonPrimitive> Primitives;
 
+	TArray<FString> TargetNames;
+
 	UE_API virtual void WriteObject(IGLTFJsonWriter& Writer) const override;
 
 	UE_API bool HasValue() const;
diff --git a/Engine/Plugins/Enterprise/GLTFExporter/Source/GLTFExporter/Public/Options/GLTFExportOptions.h b/Engine/Plugins/Enterprise/GLTFExporter/Source/GLTFExporter/Public/Options/GLTFExportOptions.h
--- a/Engine/Plugins/Enterprise/GLTFExporter/Source/GLTFExporter/Public/Options/GLTFExportOptions.h
+++ b/Engine/Plugins/Enterprise/GLTFExporter/Source/GLTFExporter/Public/Options/GLTFExportOptions.h
@@ -127,6 +127,10 @@ class UGLTFExportOptions : public UObject
 	UPROPERTY(EditAnywhere, BlueprintReadWrite, Config, Category = Mesh)
 	bool bUseMeshQuantization;
 
+	/** If enabled, exports Morph Targets. */
+	UPROPERTY(EditAnywhere, BlueprintReadWrite, Config, Category = Mesh)
+	bool bExportMorphTargets;
+
 	/** If enabled, export level sequences. Only transform tracks are currently supported. The level sequence will be played at the assigned display rate. */
 	UPROPERTY(EditAnywhere, BlueprintReadWrite, Config, Category = Animation)
 	bool bExportLevelSequences;
diff --git a/Engine/Plugins/Enterprise/GLTFExporter/Source/GLTFExporter/Public/Utilities/GLTFCoreUtilities.h b/Engine/Plugins/Enterprise/GLTFExporter/Source/GLTFExporter/Public/Utilities/GLTFCoreUtilities.h
--- a/Engine/Plugins/Enterprise/GLTFExporter/Source/GLTFExporter/Public/Utilities/GLTFCoreUtilities.h
+++ b/Engine/Plugins/Enterprise/GLTFExporter/Source/GLTFExporter/Public/Utilities/GLTFCoreUtilities.h
@@ -31,6 +31,10 @@ struct FGLTFCoreUtilities
 	static UE_API FGLTFInt8Vector4 ConvertNormal(const FPackedNormal& Normal);
 	static UE_API FGLTFInt16Vector4 ConvertNormal(const FPackedRGBA16N& Normal);
 
+	static UE_API FGLTFVector3 ConvertNormal(const FVector3f& Normal, const FVector3f& Offset);
+	static UE_API FGLTFInt8Vector4 ConvertNormal(const FPackedNormal& Normal, const FVector3f& Offset);
+	static UE_API FGLTFInt16Vector4 ConvertNormal(const FPackedRGBA16N& Normal, const FVector3f& Offset);
+
 	static UE_API FGLTFVector4 ConvertTangent(const FVector3f& Tangent, const FVector4f& Normal = FVector4f(ForceInitToZero));
 	static UE_API FGLTFVector4 ConvertTangent(const FVector4f& Tangent);
 	static UE_API FGLTFInt8Vector4 ConvertTangent(const FPackedNormal& Tangent, const FPackedNormal& Normal = FPackedNormal());
diff --git a/Engine/Plugins/Enterprise/GLTFExporter/Source/GLTFExporter/Public/Utilities/MeshParsingUtilities.h b/Engine/Plugins/Enterprise/GLTFExporter/Source/GLTFExporter/Public/Utilities/MeshParsingUtilities.h
--- a/Engine/Plugins/Enterprise/GLTFExporter/Source/GLTFExporter/Public/Utilities/MeshParsingUtilities.h
+++ b/Engine/Plugins/Enterprise/GLTFExporter/Source/GLTFExporter/Public/Utilities/MeshParsingUtilities.h
@@ -14,6 +14,15 @@ class USplineMeshComponent;
 
 namespace UE::MeshParser
 {
+	struct FPrimitiveTargetDescription
+	{
+		bool bHasPositions = false;
+		FMeshAttributesArray<FVector3f>         Positions;
+
+		bool bHasNormals = false;
+		FMeshAttributesArray<FVector3f>         Normals;
+	};
+
 	struct FMeshPrimitiveDescription
 	{
 		int32                                   MaterialIndex = INDEX_NONE;
@@ -29,11 +38,13 @@ namespace UE::MeshParser
 		TArray<FMeshAttributesArray<UE::Math::TIntVector4<FBoneIndexType>>> JointInfluences; //TODO: Per Group (make group size configurable)
 		TArray<FMeshAttributesArray<UE::Math::TIntVector4<uint16>>>         JointWeights; //TODO: Per Group (make group size configurable)
 
+		TArray<FPrimitiveTargetDescription>		Targets; //MorphTargets
+
 		bool IsEmpty();
 		void EmptyContainers();
 
 		//Except for Joint related containers
-		void PrepareContainers(const int32& IndexCount, const int32& AttributesCount, const int32& UVCount, bool bHasVertexColors);
+		void PrepareContainers(const int32& IndexCount, const int32& AttributesCount, const int32& UVCount, bool bHasVertexColors, const int32& TargetCount);
 
 		void PrepareJointContainers(const uint32& JointGroupCount, const int32& AttributesCount);
 	};
@@ -51,11 +62,13 @@ namespace UE::MeshParser
 		bool bExportVertexColors = false;
 		const USplineMeshComponent* SplineMeshComponent = nullptr;
 		int32 SkeletonInfluenceCountPerGroup = INT32_MAX;
+		bool bExportMorphTargets = false;
 
 		FExportConfigs(bool bInExportVertexSkinWeights,
 			bool bInExportVertexColors,
 			const USplineMeshComponent* InSplineMeshComponent,
-			int32 InSkeletonInfluenceCountPerGroup);
+			int32 InSkeletonInfluenceCountPerGroup,
+			bool bInExportMorphTargets);
 	};
 
 	/*
@@ -74,7 +87,7 @@ namespace UE::MeshParser
 		/*
 		* Parses FMeshDescription into the FMeshPrimitiveMeshDescriptions array based on the provided ExportConfigs.
 		*/
-		void Parse(TArray<FMeshPrimitiveDescription>& OutMeshPrimitiveDescriptions, const FExportConfigs& ExportConfigs);
+		void Parse(TArray<FMeshPrimitiveDescription>& OutMeshPrimitiveDescriptions, TArray<FString>& OutTargetNames, const FExportConfigs& ExportConfigs);
 
 	private:
 		/*
@@ -108,6 +121,12 @@ namespace UE::MeshParser
 		TVertexInstanceAttributesConstRef<FVector2f> VertexInstanceUVs;
 		TVertexInstanceAttributesConstRef<FVector4f> VertexInstanceColors;
 
+		//MorphTarget Deltas:
+		TArray<TVertexAttributesConstRef<FVector3f>> TargetVertexPositionDeltas;
+		TArray<TVertexInstanceAttributesConstRef<FVector3f>> TargetVertexInstanceNormalDeltas;
+		TArray<FString> TargetNames;
+		int32 NumOfTargets;
+
 		TConstArrayView<FVertexID> VertexInstanceIdToVertexId;
 		TPolygonGroupAttributesConstRef<FName> PolygonGroupMaterialSlotNames;
 
