[glTF Exporter] Bug fixes for Morph Target exports.
Exporting Deltas instead of absolute position/normals.

#rnx

diff --git a/Engine/Plugins/Enterprise/GLTFExporter/Source/GLTFExporter/Private/Builders/GLTFConvertBuilder.cpp b/Engine/Plugins/Enterprise/GLTFExporter/Source/GLTFExporter/Private/Builders/GLTFConvertBuilder.cpp
--- a/Engine/Plugins/Enterprise/GLTFExporter/Source/GLTFExporter/Private/Builders/GLTFConvertBuilder.cpp
+++ b/Engine/Plugins/Enterprise/GLTFExporter/Source/GLTFExporter/Private/Builders/GLTFConvertBuilder.cpp
@@ -27,14 +27,24 @@ bool FGLTFConvertBuilder::IsRootActor(const AActor* Actor) const
 	return ParentActor == nullptr || !IsSelectedActor(ParentActor);
 }
 
-FGLTFJsonAccessor* FGLTFConvertBuilder::AddUniquePositionAccessor(const FGLTFMeshSection* MeshSection, const FPositionVertexBuffer* VertexBuffer, const TMap<uint32, FVector3f>* TargetPositionDeltas)
+FGLTFJsonAccessor* FGLTFConvertBuilder::AddUniquePositionAccessor(const FGLTFMeshSection* MeshSection, const FPositionVertexBuffer* VertexBuffer)
 {
 	if (VertexBuffer == nullptr)
 	{
 		return nullptr;
 	}
 
-	return PositionBufferConverter->GetOrAdd(MeshSection, VertexBuffer, TargetPositionDeltas);
+	return PositionBufferConverter->GetOrAdd(MeshSection, VertexBuffer);
+}
+
+FGLTFJsonAccessor* FGLTFConvertBuilder::AddUniquePositionDeltaAccessor(const FGLTFMeshSection* MeshSection, const TMap<uint32, FVector3f>* TargetPositionDeltas)
+{
+	if (TargetPositionDeltas == nullptr)
+	{
+		return nullptr;
+	}
+
+	return PositionDeltaBufferConverter->GetOrAdd(MeshSection, TargetPositionDeltas);
 }
 
 FGLTFJsonAccessor* FGLTFConvertBuilder::AddUniquePositionAccessor(const FGLTFPositionArray& VertexBuffer)
@@ -57,19 +67,29 @@ FGLTFJsonAccessor* FGLTFConvertBuilder::AddUniqueColorAccessor(const FGLTFColorA
 	return ColorBufferConverterRaw->GetOrAdd(VertexColorBuffer);
 }
 
-FGLTFJsonAccessor* FGLTFConvertBuilder::AddUniqueNormalAccessor(const FGLTFMeshSection* MeshSection, const FStaticMeshVertexBuffer* VertexBuffer, const TMap<uint32, FVector3f>* TargetNormalDeltas)
+FGLTFJsonAccessor* FGLTFConvertBuilder::AddUniqueNormalAccessor(const FGLTFMeshSection* MeshSection, const FStaticMeshVertexBuffer* VertexBuffer)
 {
 	if (VertexBuffer == nullptr)
 	{
 		return nullptr;
 	}
 
-	return NormalBufferConverter->GetOrAdd(MeshSection, VertexBuffer, TargetNormalDeltas);
+	return NormalBufferConverter->GetOrAdd(MeshSection, VertexBuffer);
+}
+
+FGLTFJsonAccessor* FGLTFConvertBuilder::AddUniqueNormalDeltaAccessor(const FGLTFMeshSection* MeshSection, const TMap<uint32, FVector3f>* TargetNormalDeltas, const bool& bHighPrecision)
+{
+	if (TargetNormalDeltas == nullptr)
+	{
+		return nullptr;
+	}
+
+	return NormalDeltaBufferConverter->GetOrAdd(MeshSection, TargetNormalDeltas, bHighPrecision);
 }
 
-FGLTFJsonAccessor* FGLTFConvertBuilder::AddUniqueNormalAccessor(const FGLTFNormalArray& Normals)
+FGLTFJsonAccessor* FGLTFConvertBuilder::AddUniqueNormalAccessor(const FGLTFNormalArray& Normals, const bool& bNormalize)
 {
-	return NormalBufferConverterRaw->GetOrAdd(Normals);
+	return NormalBufferConverterRaw->GetOrAdd(Normals, bNormalize);
 }
 
 FGLTFJsonAccessor* FGLTFConvertBuilder::AddUniqueTangentAccessor(const FGLTFMeshSection* MeshSection, const FStaticMeshVertexBuffer* VertexBuffer)
diff --git a/Engine/Plugins/Enterprise/GLTFExporter/Source/GLTFExporter/Private/Converters/GLTFAccessorConverters.cpp b/Engine/Plugins/Enterprise/GLTFExporter/Source/GLTFExporter/Private/Converters/GLTFAccessorConverters.cpp
--- a/Engine/Plugins/Enterprise/GLTFExporter/Source/GLTFExporter/Private/Converters/GLTFAccessorConverters.cpp
+++ b/Engine/Plugins/Enterprise/GLTFExporter/Source/GLTFExporter/Private/Converters/GLTFAccessorConverters.cpp
@@ -19,9 +19,9 @@ namespace
 			return FGLTFCoreUtilities::ConvertNormal(Normal);
 		}
 
-		static DestinationType Convert(const SourceType& Normal, const FVector3f& Offset)
+		static DestinationType ConvertDelta(const SourceType& NormalDelta)
 		{
-			return FGLTFCoreUtilities::ConvertNormal(Normal, Offset);
+			return FGLTFCoreUtilities::ConvertNormalDelta(NormalDelta);
 		}
 	};
 
@@ -33,9 +33,9 @@ namespace
 			return FGLTFCoreUtilities::ConvertNormal(Normal.ToFVector3f());
 		}
 
-		static FGLTFVector3 Convert(const SourceType& Normal, const FVector3f& Offset)
+		static FGLTFVector3 ConvertDelta(const SourceType& NormalDelta)
 		{
-			return FGLTFCoreUtilities::ConvertNormal(Normal.ToFVector3f(), Offset);
+            return FGLTFCoreUtilities::ConvertNormalDelta(NormalDelta.ToFVector3f());
 		}
 	};
 
@@ -102,7 +102,7 @@ namespace
 	}
 }
 
-FGLTFJsonAccessor* FGLTFPositionBufferConverter::Convert(const FGLTFMeshSection* MeshSection, const FPositionVertexBuffer* VertexBuffer, const TMap<uint32, FVector3f>* TargetPositionDeltas)
+FGLTFJsonAccessor* FGLTFPositionBufferConverter::Convert(const FGLTFMeshSection* MeshSection, const FPositionVertexBuffer* VertexBuffer)
 {
 	if (VertexBuffer == nullptr || VertexBuffer->GetNumVertices() == 0)
 	{
@@ -125,58 +125,23 @@ FGLTFJsonAccessor* FGLTFPositionBufferConverter::Convert(const FGLTFMeshSection*
 	TArray<FGLTFVector3> Positions;
 	Positions.AddUninitialized(VertexCount);
 
-	if (TargetPositionDeltas && TargetPositionDeltas->Num() > 0)
-	{
-		if (IndexMap.Num() > 0)
-		{
-			for (uint32 VertexIndex = 0; VertexIndex < VertexCount; ++VertexIndex)
-			{
-				const uint32 MappedVertexIndex = IndexMap[VertexIndex];
-				FVector3f Position = *reinterpret_cast<const FVector3f*>(SourceData + Stride * MappedVertexIndex);
-
-				if (TargetPositionDeltas->Contains(MappedVertexIndex))
-				{
-					Position += TargetPositionDeltas->FindChecked(MappedVertexIndex);
-				}
-				
-				Positions[VertexIndex] = FGLTFCoreUtilities::ConvertPosition(Position, Builder.ExportOptions->ExportUniformScale);
-			}
-		}
-		else
-		{
-			for (uint32 VertexIndex = 0; VertexIndex < VertexCount; ++VertexIndex)
-			{
-				FVector3f Position = *reinterpret_cast<const FVector3f*>(SourceData + Stride * VertexIndex);
-
-				if (TargetPositionDeltas->Contains(VertexIndex))
-				{
-					Position += TargetPositionDeltas->FindChecked(VertexIndex);
-				}
-
-				Positions[VertexIndex] = FGLTFCoreUtilities::ConvertPosition(Position, Builder.ExportOptions->ExportUniformScale);
-			}
-		}
-	}
-	else
-	{
-		if (IndexMap.Num() > 0)
-		{
-			for (uint32 VertexIndex = 0; VertexIndex < VertexCount; ++VertexIndex)
-			{
-				const uint32 MappedVertexIndex = IndexMap[VertexIndex];
-				const FVector3f& Position = *reinterpret_cast<const FVector3f*>(SourceData + Stride * MappedVertexIndex);
-				Positions[VertexIndex] = FGLTFCoreUtilities::ConvertPosition(Position, Builder.ExportOptions->ExportUniformScale);
-			}
-		}
-		else
-		{
-			for (uint32 VertexIndex = 0; VertexIndex < VertexCount; ++VertexIndex)
-			{
-				const FVector3f& Position = *reinterpret_cast<const FVector3f*>(SourceData + Stride * VertexIndex);
-				Positions[VertexIndex] = FGLTFCoreUtilities::ConvertPosition(Position, Builder.ExportOptions->ExportUniformScale);
-			}
-		}
-	}
+    if (IndexMap.Num() > 0)
+    {
+        for (uint32 VertexIndex = 0; VertexIndex < VertexCount; ++VertexIndex)
+        {
+            const uint32 MappedVertexIndex = IndexMap[VertexIndex];
+            const FVector3f& Position = *reinterpret_cast<const FVector3f*>(SourceData + Stride * MappedVertexIndex);
+            Positions[VertexIndex] = FGLTFCoreUtilities::ConvertPosition(Position, Builder.ExportOptions->ExportUniformScale);
+        }
+    }
+    else
+    {
+        for (uint32 VertexIndex = 0; VertexIndex < VertexCount; ++VertexIndex)
+        {
+            const FVector3f& Position = *reinterpret_cast<const FVector3f*>(SourceData + Stride * VertexIndex);
+            Positions[VertexIndex] = FGLTFCoreUtilities::ConvertPosition(Position, Builder.ExportOptions->ExportUniformScale);
+        }
+    }
 	
 	FGLTFJsonAccessor* JsonAccessor = Builder.AddAccessor();
 	JsonAccessor->BufferView = Builder.AddBufferView(Positions, EGLTFJsonBufferTarget::ArrayBuffer);
@@ -209,6 +174,70 @@ FGLTFJsonAccessor* FGLTFPositionBufferConverter::Convert(const FGLTFMeshSection*
 	return JsonAccessor;
 }
 
+FGLTFJsonAccessor* FGLTFPositionDeltaBufferConverter::Convert(const FGLTFMeshSection* MeshSection, const TMap<uint32, FVector3f>* TargetPositionDeltas)
+{
+    const TArray<uint32>& IndexMap = MeshSection->IndexMap;
+    const uint32 VertexCount = IndexMap.Num() > 0 ? IndexMap.Num() : TargetPositionDeltas->Num();
+
+    TArray<FGLTFVector3> Positions;
+    Positions.AddUninitialized(VertexCount);
+
+    if (TargetPositionDeltas)
+    {
+        if (IndexMap.Num() > 0)
+        {
+            for (uint32 VertexIndex = 0; VertexIndex < VertexCount; ++VertexIndex)
+            {
+                const uint32 MappedVertexIndex = IndexMap[VertexIndex];
+                const FVector3f* Delta = TargetPositionDeltas->Find(MappedVertexIndex);
+                FVector3f Position = Delta ? *Delta : FVector3f();
+
+                Positions[VertexIndex] = FGLTFCoreUtilities::ConvertPosition(Position, Builder.ExportOptions->ExportUniformScale);
+            }
+        }
+        else
+        {
+            for (uint32 VertexIndex = 0; VertexIndex < VertexCount; ++VertexIndex)
+            {
+                const FVector3f* Delta = TargetPositionDeltas->Find(VertexIndex);
+                FVector3f Position = Delta ? *Delta : FVector3f();
+
+                Positions[VertexIndex] = FGLTFCoreUtilities::ConvertPosition(Position, Builder.ExportOptions->ExportUniformScale);
+            }
+        }
+    }
+
+    FGLTFJsonAccessor* JsonAccessor = Builder.AddAccessor();
+    JsonAccessor->BufferView = Builder.AddBufferView(Positions, EGLTFJsonBufferTarget::ArrayBuffer);
+    JsonAccessor->ComponentType = EGLTFJsonComponentType::Float;
+    JsonAccessor->Count = VertexCount;
+    JsonAccessor->Type = EGLTFJsonAccessorType::Vec3;
+
+    if (VertexCount > 0)
+    {
+        // Calculate accurate bounding box based on raw vertex values
+        JsonAccessor->MinMaxLength = 3;
+
+        for (int32 ComponentIndex = 0; ComponentIndex < JsonAccessor->MinMaxLength; ComponentIndex++)
+        {
+            JsonAccessor->Min[ComponentIndex] = Positions[0].Components[ComponentIndex];
+            JsonAccessor->Max[ComponentIndex] = Positions[0].Components[ComponentIndex];
+        }
+
+        for (uint32 VertexIndex = 1; VertexIndex < VertexCount; ++VertexIndex)
+        {
+            const FGLTFVector3& Position = Positions[VertexIndex];
+            for (int32 ComponentIndex = 0; ComponentIndex < JsonAccessor->MinMaxLength; ComponentIndex++)
+            {
+                JsonAccessor->Min[ComponentIndex] = FMath::Min(JsonAccessor->Min[ComponentIndex], Position.Components[ComponentIndex]);
+                JsonAccessor->Max[ComponentIndex] = FMath::Max(JsonAccessor->Max[ComponentIndex], Position.Components[ComponentIndex]);
+            }
+        }
+    }
+
+    return JsonAccessor;
+}
+
 FGLTFJsonAccessor* FGLTFColorBufferConverter::Convert(const FGLTFMeshSection* MeshSection, const FColorVertexBuffer* VertexColorBufferBuffer)
 {
 	if (VertexColorBufferBuffer == nullptr || VertexColorBufferBuffer->GetNumVertices() == 0)
@@ -262,7 +291,7 @@ FGLTFJsonAccessor* FGLTFColorBufferConverter::Convert(const FGLTFMeshSection* Me
 	return JsonAccessor;
 }
 
-FGLTFJsonAccessor* FGLTFNormalBufferConverter::Convert(const FGLTFMeshSection* MeshSection, const FStaticMeshVertexBuffer* VertexBuffer, const TMap<uint32, FVector3f>* TargetNormalDeltas)
+FGLTFJsonAccessor* FGLTFNormalBufferConverter::Convert(const FGLTFMeshSection* MeshSection, const FStaticMeshVertexBuffer* VertexBuffer)
 {
 	if (VertexBuffer == nullptr || VertexBuffer->GetNumVertices() == 0)
 	{
@@ -291,13 +320,13 @@ FGLTFJsonAccessor* FGLTFNormalBufferConverter::Convert(const FGLTFMeshSection* M
 		if (bHighPrecision)
 		{
 			ComponentType = EGLTFJsonComponentType::Int16;
-			BufferView = ConvertBufferView<FGLTFInt16Vector4, FPackedRGBA16N>(MeshSection, SourceData, TargetNormalDeltas);
+			BufferView = ConvertBufferView<FGLTFInt16Vector4, FPackedRGBA16N>(MeshSection, SourceData);
 			BufferView->ByteStride = sizeof(FGLTFInt16Vector4);
 		}
 		else
 		{
 			ComponentType = EGLTFJsonComponentType::Int8;
-			BufferView = ConvertBufferView<FGLTFInt8Vector4, FPackedNormal>(MeshSection, SourceData, TargetNormalDeltas);
+			BufferView = ConvertBufferView<FGLTFInt8Vector4, FPackedNormal>(MeshSection, SourceData);
 			BufferView->ByteStride = sizeof(FGLTFInt8Vector4);
 		}
 	}
@@ -305,8 +334,8 @@ FGLTFJsonAccessor* FGLTFNormalBufferConverter::Convert(const FGLTFMeshSection* M
 	{
 		ComponentType = EGLTFJsonComponentType::Float;
 		BufferView = bHighPrecision
-			? ConvertBufferView<FGLTFVector3, FPackedRGBA16N>(MeshSection, SourceData, TargetNormalDeltas)
-			: ConvertBufferView<FGLTFVector3, FPackedNormal>(MeshSection, SourceData, TargetNormalDeltas);
+			? ConvertBufferView<FGLTFVector3, FPackedRGBA16N>(MeshSection, SourceData)
+			: ConvertBufferView<FGLTFVector3, FPackedNormal>(MeshSection, SourceData);
 	}
 
 	FGLTFJsonAccessor* JsonAccessor = Builder.AddAccessor();
@@ -320,7 +349,7 @@ FGLTFJsonAccessor* FGLTFNormalBufferConverter::Convert(const FGLTFMeshSection* M
 }
 
 template <typename DestinationType, typename SourceType>
-FGLTFJsonBufferView* FGLTFNormalBufferConverter::ConvertBufferView(const FGLTFMeshSection* MeshSection, const void* SourceData, const TMap<uint32, FVector3f>* TargetNormalDeltas)
+FGLTFJsonBufferView* FGLTFNormalBufferConverter::ConvertBufferView(const FGLTFMeshSection* MeshSection, const void* SourceData)
 {
 	const TArray<uint32>& IndexMap = MeshSection->IndexMap;
 	const uint32 VertexCount = IndexMap.Num();
@@ -329,40 +358,87 @@ FGLTFJsonBufferView* FGLTFNormalBufferConverter::ConvertBufferView(const FGLTFMe
 	const TangentDatumType* TangentData = static_cast<const TangentDatumType*>(SourceData);
 
 	TArray<DestinationType> Normals;
-	Normals.AddUninitialized(VertexCount);
+    Normals.SetNumZeroed(VertexCount);
 
-	if (TargetNormalDeltas && TargetNormalDeltas->Num())
-	{
-		for (uint32 VertexIndex = 0; VertexIndex < VertexCount; ++VertexIndex)
-		{
-			const uint32 MappedVertexIndex = IndexMap[VertexIndex];
-			const TangentDatumType& TangentDatum = TangentData[MappedVertexIndex];
-			
-
-			const FVector3f* TargetNormalDelta = TargetNormalDeltas->Find(MappedVertexIndex);
-			if (TargetNormalDelta)
-			{
-				Normals[VertexIndex] = TGLTFVertexNormalUtilities<DestinationType, SourceType>::Convert(TangentDatum.TangentZ, *TargetNormalDelta);
-			}
-			else
-			{
-				Normals[VertexIndex] = TGLTFVertexNormalUtilities<DestinationType, SourceType>::Convert(TangentDatum.TangentZ);
-			}
-		}
-	}
-	else
-	{
-		for (uint32 VertexIndex = 0; VertexIndex < VertexCount; ++VertexIndex)
-		{
-			const uint32 MappedVertexIndex = IndexMap[VertexIndex];
-			const TangentDatumType& TangentDatum = TangentData[MappedVertexIndex];
-			Normals[VertexIndex] = TGLTFVertexNormalUtilities<DestinationType, SourceType>::Convert(TangentDatum.TangentZ);
-		}
-	}
+    for (uint32 VertexIndex = 0; VertexIndex < VertexCount; ++VertexIndex)
+    {
+        const uint32 MappedVertexIndex = IndexMap[VertexIndex];
+        const TangentDatumType& TangentDatum = TangentData[MappedVertexIndex];
+        Normals[VertexIndex] = TGLTFVertexNormalUtilities<DestinationType, SourceType>::Convert(TangentDatum.TangentZ);
+    }
 
 	return Builder.AddBufferView(Normals, EGLTFJsonBufferTarget::ArrayBuffer);
 }
 
+
+FGLTFJsonAccessor* FGLTFNormalDeltaBufferConverter::Convert(const FGLTFMeshSection* MeshSection, const TMap<uint32, FVector3f>* TargetNormalDeltas, bool bHighPrecision)
+{
+
+    FGLTFJsonBufferView* BufferView;
+    EGLTFJsonComponentType ComponentType;
+
+    const bool bMeshQuantization = Builder.ExportOptions->bUseMeshQuantization;
+    if (bMeshQuantization)
+    {
+        Builder.AddExtension(EGLTFJsonExtension::KHR_MeshQuantization, true);
+
+        if (bHighPrecision)
+        {
+            ComponentType = EGLTFJsonComponentType::Int16;
+            BufferView = ConvertBufferView<FGLTFInt16Vector4, FPackedRGBA16N>(MeshSection, TargetNormalDeltas);
+            BufferView->ByteStride = sizeof(FGLTFInt16Vector4);
+        }
+        else
+        {
+            ComponentType = EGLTFJsonComponentType::Int8;
+            BufferView = ConvertBufferView<FGLTFInt8Vector4, FPackedNormal>(MeshSection, TargetNormalDeltas);
+            BufferView->ByteStride = sizeof(FGLTFInt8Vector4);
+        }
+    }
+    else
+    {
+        ComponentType = EGLTFJsonComponentType::Float;
+        BufferView = bHighPrecision
+            ? ConvertBufferView<FGLTFVector3, FPackedRGBA16N>(MeshSection, TargetNormalDeltas)
+            : ConvertBufferView<FGLTFVector3, FPackedNormal>(MeshSection, TargetNormalDeltas);
+    }
+
+    FGLTFJsonAccessor* JsonAccessor = Builder.AddAccessor();
+    JsonAccessor->BufferView = BufferView;
+    JsonAccessor->ComponentType = ComponentType;
+    JsonAccessor->Count = MeshSection->IndexMap.Num();
+    JsonAccessor->Type = EGLTFJsonAccessorType::Vec3;
+    JsonAccessor->bNormalized = bMeshQuantization;
+
+    return JsonAccessor;
+}
+
+template <typename DestinationType, typename SourceType>
+FGLTFJsonBufferView* FGLTFNormalDeltaBufferConverter::ConvertBufferView(const FGLTFMeshSection* MeshSection, const TMap<uint32, FVector3f>* TargetNormalDeltas)
+{
+    const TArray<uint32>& IndexMap = MeshSection->IndexMap;
+    const uint32 VertexCount = IndexMap.Num();
+
+    TArray<DestinationType> Normals;
+    Normals.SetNumZeroed(VertexCount);
+
+    if (TargetNormalDeltas && TargetNormalDeltas->Num())
+    {
+        for (uint32 VertexIndex = 0; VertexIndex < VertexCount; ++VertexIndex)
+        {
+            const uint32 MappedVertexIndex = IndexMap[VertexIndex];
+
+            const FVector3f* TargetNormalDelta = TargetNormalDeltas->Find(MappedVertexIndex);
+            if (TargetNormalDelta)
+            {
+                Normals[VertexIndex] = TGLTFVertexNormalUtilities<DestinationType, SourceType>::ConvertDelta(*TargetNormalDelta);
+            }
+        }
+    }
+
+    return Builder.AddBufferView(Normals, EGLTFJsonBufferTarget::ArrayBuffer);
+}
+
 FGLTFJsonAccessor* FGLTFTangentBufferConverter::Convert(const FGLTFMeshSection* MeshSection, const FStaticMeshVertexBuffer* VertexBuffer)
 {
 	if (VertexBuffer == nullptr || VertexBuffer->GetNumVertices() == 0)
@@ -767,17 +843,28 @@ FGLTFJsonAccessor* FGLTFIndexBufferConverterRaw::Convert(FGLTFIndexArray IndexBu
 }
 
 
-FGLTFJsonAccessor* FGLTFNormalBufferConverterRaw::Convert(FGLTFNormalArray NormalsSource)
+FGLTFJsonAccessor* FGLTFNormalBufferConverterRaw::Convert(FGLTFNormalArray NormalsSource, bool bNormalize)
 {
 	const uint32 VertexCount = NormalsSource.Num();
 
 	TArray<FGLTFVector3> Normals;
 	Normals.AddUninitialized(VertexCount);
 
-	for (uint32 VertexIndex = 0; VertexIndex < VertexCount; ++VertexIndex)
-	{
-		Normals[VertexIndex] = FGLTFCoreUtilities::ConvertNormal(NormalsSource[VertexIndex]);
-	}
+    if (bNormalize)
+    {
+        for (uint32 VertexIndex = 0; VertexIndex < VertexCount; ++VertexIndex)
+        {
+            Normals[VertexIndex] = FGLTFCoreUtilities::ConvertNormal(NormalsSource[VertexIndex]);
+        }
+    }
+    else
+    {
+        for (uint32 VertexIndex = 0; VertexIndex < VertexCount; ++VertexIndex)
+        {
+            Normals[VertexIndex] = FGLTFCoreUtilities::ConvertVector(NormalsSource[VertexIndex]);
+        }
+    }
+	
 
 	return ConvertToAccessor<float, EGLTFJsonAccessorType::Vec3>(Builder, Normals);
 }
diff --git a/Engine/Plugins/Enterprise/GLTFExporter/Source/GLTFExporter/Private/Tasks/GLTFDelayedMeshTasks.cpp b/Engine/Plugins/Enterprise/GLTFExporter/Source/GLTFExporter/Private/Tasks/GLTFDelayedMeshTasks.cpp
--- a/Engine/Plugins/Enterprise/GLTFExporter/Source/GLTFExporter/Private/Tasks/GLTFDelayedMeshTasks.cpp
+++ b/Engine/Plugins/Enterprise/GLTFExporter/Source/GLTFExporter/Private/Tasks/GLTFDelayedMeshTasks.cpp
@@ -461,7 +461,7 @@ void FGLTFDelayedSkeletalMeshTask::ProcessSourceModel(const TArray<FSkeletalMate
 			{
 				FGLTFJsonTarget& JSONMorphTarget = JsonPrimitive.Targets.AddDefaulted_GetRef();
 				if (PrimitiveTargetDescription.bHasPositions) JSONMorphTarget.Position = Builder.AddUniquePositionAccessor(PrimitiveTargetDescription.Positions);
-				if (PrimitiveTargetDescription.bHasNormals) JSONMorphTarget.Normal = Builder.AddUniqueNormalAccessor(PrimitiveTargetDescription.Normals);
+				if (PrimitiveTargetDescription.bHasNormals) JSONMorphTarget.Normal = Builder.AddUniqueNormalAccessor(PrimitiveTargetDescription.Normals, false);
 			}
 		}
 
@@ -593,6 +593,8 @@ void FGLTFDelayedSkeletalMeshTask::ProcessRenderData(const TArray<FSkeletalMater
 				));
 		}
 
+		const bool bHighPrecision = VertexBuffer.GetUseHighPrecisionTangentBasis();
+
 		if (MorphTargetsPositions.Num() > 0)
 		{
 			JsonPrimitive.Targets.Reserve(MorphTargetsPositions.Num());
@@ -605,8 +607,8 @@ void FGLTFDelayedSkeletalMeshTask::ProcessRenderData(const TArray<FSkeletalMater
 				if (PositionsDeltas.bIsValid || NormalsDeltas.bIsValid)
 				{
 					FGLTFJsonTarget& JSONMorphTarget = JsonPrimitive.Targets.AddDefaulted_GetRef();
-					if (PositionsDeltas.bIsValid) JSONMorphTarget.Position = Builder.AddUniquePositionAccessor(ConvertedSection, &PositionBuffer, &PositionsDeltas.Deltas);
-					if (NormalsDeltas.bIsValid) JSONMorphTarget.Normal = Builder.AddUniqueNormalAccessor(ConvertedSection, &VertexBuffer, &NormalsDeltas.Deltas);
+					if (PositionsDeltas.bIsValid) JSONMorphTarget.Position = Builder.AddUniquePositionDeltaAccessor(ConvertedSection, &PositionsDeltas.Deltas);
+					if (NormalsDeltas.bIsValid) JSONMorphTarget.Normal = Builder.AddUniqueNormalDeltaAccessor(ConvertedSection, &NormalsDeltas.Deltas, bHighPrecision);
 				}
 			}
 		}
diff --git a/Engine/Plugins/Enterprise/GLTFExporter/Source/GLTFExporter/Private/Utilities/GLTFCoreUtilities.cpp b/Engine/Plugins/Enterprise/GLTFExporter/Source/GLTFExporter/Private/Utilities/GLTFCoreUtilities.cpp
--- a/Engine/Plugins/Enterprise/GLTFExporter/Source/GLTFExporter/Private/Utilities/GLTFCoreUtilities.cpp
+++ b/Engine/Plugins/Enterprise/GLTFExporter/Source/GLTFExporter/Private/Utilities/GLTFCoreUtilities.cpp
@@ -43,22 +43,19 @@ FGLTFInt16Vector4 FGLTFCoreUtilities::ConvertNormal(const FPackedRGBA16N& Normal
 	return { SafeNormal.X, SafeNormal.Z, SafeNormal.Y, 0 };
 }
 
-FGLTFVector3 FGLTFCoreUtilities::ConvertNormal(const FVector3f& Normal, const FVector3f& Offset)
+FGLTFVector3 FGLTFCoreUtilities::ConvertNormalDelta(const FVector3f& NormalDelta)
 {
-	const FVector3f SafeNormal = (Normal + Offset).GetSafeNormal();
-	return ConvertVector(SafeNormal);
+	return ConvertVector(NormalDelta);
 }
 
-FGLTFInt8Vector4 FGLTFCoreUtilities::ConvertNormal(const FPackedNormal& Normal, const FVector3f& Offset)
+FGLTFInt8Vector4 FGLTFCoreUtilities::ConvertNormalDelta(const FPackedNormal& NormalDelta)
 {
-	const FPackedNormal SafeNormal = (Normal.ToFVector3f() + Offset).GetSafeNormal();
-	return { SafeNormal.Vector.X, SafeNormal.Vector.Z, SafeNormal.Vector.Y, 0 };
+	return { NormalDelta.Vector.X, NormalDelta.Vector.Z, NormalDelta.Vector.Y, 0 };
 }
 
-FGLTFInt16Vector4 FGLTFCoreUtilities::ConvertNormal(const FPackedRGBA16N& Normal, const FVector3f& Offset)
+FGLTFInt16Vector4 FGLTFCoreUtilities::ConvertNormalDelta(const FPackedRGBA16N& NormalDelta)
 {
-	const FPackedRGBA16N SafeNormal = (Normal.ToFVector3f() + Offset).GetSafeNormal();
-	return { SafeNormal.X, SafeNormal.Z, SafeNormal.Y, 0 };
+	return { NormalDelta.X, NormalDelta.Z, NormalDelta.Y, 0 };
 }
 
 
diff --git a/Engine/Plugins/Enterprise/GLTFExporter/Source/GLTFExporter/Private/Utilities/MeshParsingUtilities.cpp b/Engine/Plugins/Enterprise/GLTFExporter/Source/GLTFExporter/Private/Utilities/MeshParsingUtilities.cpp
--- a/Engine/Plugins/Enterprise/GLTFExporter/Source/GLTFExporter/Private/Utilities/MeshParsingUtilities.cpp
+++ b/Engine/Plugins/Enterprise/GLTFExporter/Source/GLTFExporter/Private/Utilities/MeshParsingUtilities.cpp
@@ -404,11 +404,11 @@ namespace UE::MeshParser
 							FVertexID OriginalVertexID = VertexIdTracker.VertexId;
 							FVertexInstanceID VertexInstanceID = VertexIdTracker.VertexInstanceId;
 
-							MeshPrimitiveDescription.Targets[TargetCounter].Positions[ExportedVertexID] = MeshPrimitiveDescription.Positions[ExportedVertexID] + TargetVertexPositionDeltas[TargetCounter][OriginalVertexID];
+							MeshPrimitiveDescription.Targets[TargetCounter].Positions[ExportedVertexID] = TargetVertexPositionDeltas[TargetCounter][OriginalVertexID];
 
 							if (!CheckCompareSetValue(AttributesSetTracker.TargetNormals[TargetCounter][ExportedVertexID],
 								MeshPrimitiveDescription.Targets[TargetCounter].Normals[ExportedVertexID],
-								VertexInstanceNormals[VertexInstanceID] + TargetVertexInstanceNormalDeltas[TargetCounter][OriginalVertexID]))
+								TargetVertexInstanceNormalDeltas[TargetCounter][VertexInstanceID]))
 							{
 								return false;
 							}
@@ -423,7 +423,7 @@ namespace UE::MeshParser
 							FVertexID OriginalVertexID = VertexIdTracker.VertexId;
 							FVertexInstanceID VertexInstanceID = VertexIdTracker.VertexInstanceId;
 
-							MeshPrimitiveDescription.Targets[TargetCounter].Positions[ExportedVertexID] = MeshPrimitiveDescription.Positions[ExportedVertexID] + TargetVertexPositionDeltas[TargetCounter][OriginalVertexID];
+							MeshPrimitiveDescription.Targets[TargetCounter].Positions[ExportedVertexID] = TargetVertexPositionDeltas[TargetCounter][OriginalVertexID];
 						}
 					}
 				}
@@ -439,7 +439,7 @@ namespace UE::MeshParser
 
 						if (!CheckCompareSetValue(AttributesSetTracker.TargetNormals[TargetCounter][ExportedVertexID],
 							MeshPrimitiveDescription.Targets[TargetCounter].Normals[ExportedVertexID],
-							VertexInstanceNormals[VertexInstanceID] + TargetVertexInstanceNormalDeltas[TargetCounter][OriginalVertexID]))
+							TargetVertexInstanceNormalDeltas[TargetCounter][VertexInstanceID]))
 						{
 							return false;
 						}
@@ -617,8 +617,8 @@ namespace UE::MeshParser
 							FVertexID OriginalVertexID = VertexIdTracker.VertexId;
 							FVertexInstanceID VertexInstanceID = VertexIdTracker.VertexInstanceId;
 
-							MeshPrimitiveDescription.Targets[TargetCounter].Positions[ExportedVertexID] = VertexPositions[OriginalVertexID] + (TargetVertexPositionDeltas[TargetCounter].IsValid() ? TargetVertexPositionDeltas[TargetCounter][OriginalVertexID] : FVector3f());
-							MeshPrimitiveDescription.Targets[TargetCounter].Normals[ExportedVertexID] = VertexInstanceNormals[VertexInstanceID] + (TargetVertexInstanceNormalDeltas[TargetCounter].IsValid() ? TargetVertexInstanceNormalDeltas[TargetCounter][VertexInstanceID] : FVector3f());
+							MeshPrimitiveDescription.Targets[TargetCounter].Positions[ExportedVertexID] = TargetVertexPositionDeltas[TargetCounter][OriginalVertexID];
+							MeshPrimitiveDescription.Targets[TargetCounter].Normals[ExportedVertexID] = TargetVertexInstanceNormalDeltas[TargetCounter][VertexInstanceID];
 						}
 					}
 					else
@@ -630,7 +630,7 @@ namespace UE::MeshParser
 							FVertexID OriginalVertexID = VertexIdTracker.VertexId;
 							FVertexInstanceID VertexInstanceID = VertexIdTracker.VertexInstanceId;
 
-							MeshPrimitiveDescription.Targets[TargetCounter].Positions[ExportedVertexID] = VertexPositions[OriginalVertexID] + (TargetVertexPositionDeltas[TargetCounter].IsValid() ? TargetVertexPositionDeltas[TargetCounter][OriginalVertexID] : FVector3f());
+							MeshPrimitiveDescription.Targets[TargetCounter].Positions[ExportedVertexID] = TargetVertexPositionDeltas[TargetCounter][OriginalVertexID];
 						}
 					}
 				}
@@ -643,7 +643,7 @@ namespace UE::MeshParser
 						FVertexID OriginalVertexID = VertexIdTracker.VertexId;
 						FVertexInstanceID VertexInstanceID = VertexIdTracker.VertexInstanceId;
 
-						MeshPrimitiveDescription.Targets[TargetCounter].Normals[ExportedVertexID] = VertexInstanceNormals[VertexInstanceID] + (TargetVertexInstanceNormalDeltas[TargetCounter].IsValid() ? TargetVertexInstanceNormalDeltas[TargetCounter][VertexInstanceID] : FVector3f());
+						MeshPrimitiveDescription.Targets[TargetCounter].Normals[ExportedVertexID] = TargetVertexInstanceNormalDeltas[TargetCounter][VertexInstanceID];
 					}
 				}
 			}
diff --git a/Engine/Plugins/Enterprise/GLTFExporter/Source/GLTFExporter/Public/Builders/GLTFConvertBuilder.h b/Engine/Plugins/Enterprise/GLTFExporter/Source/GLTFExporter/Public/Builders/GLTFConvertBuilder.h
--- a/Engine/Plugins/Enterprise/GLTFExporter/Source/GLTFExporter/Public/Builders/GLTFConvertBuilder.h
+++ b/Engine/Plugins/Enterprise/GLTFExporter/Source/GLTFExporter/Public/Builders/GLTFConvertBuilder.h
@@ -36,12 +36,14 @@ public:
 	UE_API bool IsSelectedActor(const AActor* Object) const;
 	UE_API bool IsRootActor(const AActor* Actor) const;
 
-	UE_API FGLTFJsonAccessor* AddUniquePositionAccessor(const FGLTFMeshSection* MeshSection, const FPositionVertexBuffer* VertexBuffer, const TMap<uint32, FVector3f>* TargetPositionDeltas = nullptr);
+	UE_API FGLTFJsonAccessor* AddUniquePositionAccessor(const FGLTFMeshSection* MeshSection, const FPositionVertexBuffer* VertexBuffer);
+	UE_API FGLTFJsonAccessor* AddUniquePositionDeltaAccessor(const FGLTFMeshSection* MeshSection, const TMap<uint32, FVector3f>* TargetPositionDeltas);
 	UE_API FGLTFJsonAccessor* AddUniquePositionAccessor(const FGLTFPositionArray& VertexBuffer);
 	UE_API FGLTFJsonAccessor* AddUniqueColorAccessor(const FGLTFMeshSection* MeshSection, const FColorVertexBuffer* VertexBuffer);
 	UE_API FGLTFJsonAccessor* AddUniqueColorAccessor(const FGLTFColorArray& VertexColorBuffer);
-	UE_API FGLTFJsonAccessor* AddUniqueNormalAccessor(const FGLTFMeshSection* MeshSection, const FStaticMeshVertexBuffer* VertexBuffer, const TMap<uint32, FVector3f>* TargetNormalDeltas = nullptr);
-	UE_API FGLTFJsonAccessor* AddUniqueNormalAccessor(const FGLTFNormalArray& Normals);
+	UE_API FGLTFJsonAccessor* AddUniqueNormalAccessor(const FGLTFMeshSection* MeshSection, const FStaticMeshVertexBuffer* VertexBuffer);
+	UE_API FGLTFJsonAccessor* AddUniqueNormalDeltaAccessor(const FGLTFMeshSection* MeshSection, const TMap<uint32, FVector3f>* TargetNormalDeltas, const bool& bHighPrecision);
+	UE_API FGLTFJsonAccessor* AddUniqueNormalAccessor(const FGLTFNormalArray& Normals, const bool& bNormalize = true);
 	UE_API FGLTFJsonAccessor* AddUniqueTangentAccessor(const FGLTFMeshSection* MeshSection, const FStaticMeshVertexBuffer* VertexBuffer);
 	UE_API FGLTFJsonAccessor* AddUniqueTangentAccessor(const FGLTFTangentArray& Tangents);
 	UE_API FGLTFJsonAccessor* AddUniqueUVAccessor(const FGLTFMeshSection* MeshSection, const FStaticMeshVertexBuffer* VertexBuffer, int32 UVIndex);
@@ -112,8 +114,10 @@ public:
 
 	//'Original' converters:
 	TUniquePtr<IGLTFPositionBufferConverter> PositionBufferConverter = MakeUnique<FGLTFPositionBufferConverter>(*this);
+	TUniquePtr<IGLTFPositionDeltaBufferConverter> PositionDeltaBufferConverter = MakeUnique<FGLTFPositionDeltaBufferConverter>(*this);
 	TUniquePtr<IGLTFColorBufferConverter> ColorBufferConverter = MakeUnique<FGLTFColorBufferConverter>(*this);
 	TUniquePtr<IGLTFNormalBufferConverter> NormalBufferConverter = MakeUnique<FGLTFNormalBufferConverter>(*this);
+	TUniquePtr<IGLTFNormalDeltaBufferConverter> NormalDeltaBufferConverter = MakeUnique<FGLTFNormalDeltaBufferConverter>(*this);
 	TUniquePtr<IGLTFTangentBufferConverter> TangentBufferConverter = MakeUnique<FGLTFTangentBufferConverter>(*this);
 	TUniquePtr<IGLTFUVBufferConverter> UVBufferConverter = MakeUnique<FGLTFUVBufferConverter>(*this);
 	TUniquePtr<IGLTFBoneIndexBufferConverter> BoneIndexBufferConverter = MakeUnique<FGLTFBoneIndexBufferConverter>(*this);
diff --git a/Engine/Plugins/Enterprise/GLTFExporter/Source/GLTFExporter/Public/Converters/GLTFAccessorConverters.h b/Engine/Plugins/Enterprise/GLTFExporter/Source/GLTFExporter/Public/Converters/GLTFAccessorConverters.h
--- a/Engine/Plugins/Enterprise/GLTFExporter/Source/GLTFExporter/Public/Converters/GLTFAccessorConverters.h
+++ b/Engine/Plugins/Enterprise/GLTFExporter/Source/GLTFExporter/Public/Converters/GLTFAccessorConverters.h
@@ -15,9 +15,11 @@ class FColorVertexBuffer;
 class FStaticMeshVertexBuffer;
 class FSkinWeightVertexBuffer;
 
-typedef TGLTFConverter<FGLTFJsonAccessor*, const FGLTFMeshSection*, const FPositionVertexBuffer*, const TMap<uint32, FVector3f>*> IGLTFPositionBufferConverter;
+typedef TGLTFConverter<FGLTFJsonAccessor*, const FGLTFMeshSection*, const FPositionVertexBuffer*> IGLTFPositionBufferConverter;
+typedef TGLTFConverter<FGLTFJsonAccessor*, const FGLTFMeshSection*, const TMap<uint32, FVector3f>*> IGLTFPositionDeltaBufferConverter;
 typedef TGLTFConverter<FGLTFJsonAccessor*, const FGLTFMeshSection*, const FColorVertexBuffer*> IGLTFColorBufferConverter;
-typedef TGLTFConverter<FGLTFJsonAccessor*, const FGLTFMeshSection*, const FStaticMeshVertexBuffer*, const TMap<uint32, FVector3f>*> IGLTFNormalBufferConverter;
+typedef TGLTFConverter<FGLTFJsonAccessor*, const FGLTFMeshSection*, const FStaticMeshVertexBuffer*> IGLTFNormalBufferConverter;
+typedef TGLTFConverter<FGLTFJsonAccessor*, const FGLTFMeshSection*, const TMap<uint32, FVector3f>*, bool> IGLTFNormalDeltaBufferConverter;
 typedef TGLTFConverter<FGLTFJsonAccessor*, const FGLTFMeshSection*, const FStaticMeshVertexBuffer*> IGLTFTangentBufferConverter;
 typedef TGLTFConverter<FGLTFJsonAccessor*, const FGLTFMeshSection*, const FStaticMeshVertexBuffer*, uint32> IGLTFUVBufferConverter;
 typedef TGLTFConverter<FGLTFJsonAccessor*, const FGLTFMeshSection*, const FSkinWeightVertexBuffer*, uint32> IGLTFBoneIndexBufferConverter;
@@ -26,7 +28,7 @@ typedef TGLTFConverter<FGLTFJsonAccessor*, const FGLTFMeshSection*> IGLTFIndexBu
 
 typedef TGLTFConverter<FGLTFJsonAccessor*, FGLTFPositionArray> IGLTFPositionBufferConverterRaw;
 typedef TGLTFConverter<FGLTFJsonAccessor*, FGLTFIndexArray, FString> IGLTFIndexBufferConverterRaw;
-typedef TGLTFConverter<FGLTFJsonAccessor*, FGLTFNormalArray> IGLTFNormalBufferConverterRaw;
+typedef TGLTFConverter<FGLTFJsonAccessor*, FGLTFNormalArray, bool> IGLTFNormalBufferConverterRaw;
 typedef TGLTFConverter<FGLTFJsonAccessor*, FGLTFUVArray> IGLTFUVBufferConverterRaw;
 typedef TGLTFConverter<FGLTFJsonAccessor*, FGLTFColorArray> IGLTFColorBufferConverterRaw;
 typedef TGLTFConverter<FGLTFJsonAccessor*, FGLTFTangentArray> IGLTFTangentBufferConverterRaw;
@@ -42,7 +44,18 @@ public:
 
 protected:
 
-	UE_API virtual FGLTFJsonAccessor* Convert(const FGLTFMeshSection* MeshSection, const FPositionVertexBuffer* VertexBuffer, const TMap<uint32, FVector3f>* TargetPositionDeltas) override;
+	UE_API virtual FGLTFJsonAccessor* Convert(const FGLTFMeshSection* MeshSection, const FPositionVertexBuffer* VertexBuffer) override;
+};
+
+class FGLTFPositionDeltaBufferConverter : public FGLTFBuilderContext, public IGLTFPositionDeltaBufferConverter
+{
+public:
+
+	using FGLTFBuilderContext::FGLTFBuilderContext;
+
+protected:
+
+	UE_API virtual FGLTFJsonAccessor* Convert(const FGLTFMeshSection* MeshSection, const TMap<uint32, FVector3f>* TargetPositionDeltas) override;
 };
 
 class FGLTFColorBufferConverter : public FGLTFBuilderContext, public IGLTFColorBufferConverter
@@ -64,12 +77,28 @@ public:
 
 protected:
 
-	UE_API virtual FGLTFJsonAccessor* Convert(const FGLTFMeshSection* MeshSection, const FStaticMeshVertexBuffer* VertexBuffer, const TMap<uint32, FVector3f>* TargetNormalDeltas) override;
+	UE_API virtual FGLTFJsonAccessor* Convert(const FGLTFMeshSection* MeshSection, const FStaticMeshVertexBuffer* VertexBuffer) override;
+
+private:
+
+	template <typename DestinationType, typename SourceType>
+	FGLTFJsonBufferView* ConvertBufferView(const FGLTFMeshSection* MeshSection, const void* TangentData);
+};
+
+class FGLTFNormalDeltaBufferConverter : public FGLTFBuilderContext, public IGLTFNormalDeltaBufferConverter
+{
+public:
+
+	using FGLTFBuilderContext::FGLTFBuilderContext;
+
+protected:
+
+	UE_API virtual FGLTFJsonAccessor* Convert(const FGLTFMeshSection* MeshSection, const TMap<uint32, FVector3f>* TargetNormalDeltas, bool bHighPrecision) override;
 
 private:
 
 	template <typename DestinationType, typename SourceType>
-	FGLTFJsonBufferView* ConvertBufferView(const FGLTFMeshSection* MeshSection, const void* TangentData, const TMap<uint32, FVector3f>* TargetNormalDeltas);
+	FGLTFJsonBufferView* ConvertBufferView(const FGLTFMeshSection* MeshSection, const TMap<uint32, FVector3f>* TargetNormalDeltas);
 };
 
 class FGLTFTangentBufferConverter : public FGLTFBuilderContext, public IGLTFTangentBufferConverter
@@ -181,7 +210,7 @@ public:
 	using FGLTFBuilderContext::FGLTFBuilderContext;
 
 protected:
-	UE_API virtual FGLTFJsonAccessor* Convert(FGLTFNormalArray NormalsSource) override;
+	UE_API virtual FGLTFJsonAccessor* Convert(FGLTFNormalArray NormalsSource, bool bNormalize) override;
 };
 
 class FGLTFUVBufferConverterRaw : public FGLTFBuilderContext, public IGLTFUVBufferConverterRaw
diff --git a/Engine/Plugins/Enterprise/GLTFExporter/Source/GLTFExporter/Public/Utilities/GLTFCoreUtilities.h b/Engine/Plugins/Enterprise/GLTFExporter/Source/GLTFExporter/Public/Utilities/GLTFCoreUtilities.h
--- a/Engine/Plugins/Enterprise/GLTFExporter/Source/GLTFExporter/Public/Utilities/GLTFCoreUtilities.h
+++ b/Engine/Plugins/Enterprise/GLTFExporter/Source/GLTFExporter/Public/Utilities/GLTFCoreUtilities.h
@@ -31,9 +31,9 @@ struct FGLTFCoreUtilities
 	static UE_API FGLTFInt8Vector4 ConvertNormal(const FPackedNormal& Normal);
 	static UE_API FGLTFInt16Vector4 ConvertNormal(const FPackedRGBA16N& Normal);
 
-	static UE_API FGLTFVector3 ConvertNormal(const FVector3f& Normal, const FVector3f& Offset);
-	static UE_API FGLTFInt8Vector4 ConvertNormal(const FPackedNormal& Normal, const FVector3f& Offset);
-	static UE_API FGLTFInt16Vector4 ConvertNormal(const FPackedRGBA16N& Normal, const FVector3f& Offset);
+	static UE_API FGLTFVector3 ConvertNormalDelta(const FVector3f& NormalDelta);
+	static UE_API FGLTFInt8Vector4 ConvertNormalDelta(const FPackedNormal& NormalDelta);
+	static UE_API FGLTFInt16Vector4 ConvertNormalDelta(const FPackedRGBA16N& NormalDelta);
 
 	static UE_API FGLTFVector4 ConvertTangent(const FVector3f& Tangent, const FVector4f& Normal = FVector4f(ForceInitToZero));
 	static UE_API FGLTFVector4 ConvertTangent(const FVector4f& Tangent);
