[glTF Exporter] MorphTarget safe guard fixes.

#rnx

diff --git a/Engine/Plugins/Enterprise/GLTFExporter/Source/GLTFExporter/Private/Converters/GLTFAccessorConverters.cpp b/Engine/Plugins/Enterprise/GLTFExporter/Source/GLTFExporter/Private/Converters/GLTFAccessorConverters.cpp
--- a/Engine/Plugins/Enterprise/GLTFExporter/Source/GLTFExporter/Private/Converters/GLTFAccessorConverters.cpp
+++ b/Engine/Plugins/Enterprise/GLTFExporter/Source/GLTFExporter/Private/Converters/GLTFAccessorConverters.cpp
@@ -151,7 +151,6 @@ FGLTFJsonAccessor* FGLTFPositionBufferConverter::Convert(const FGLTFMeshSection*
 
 	if (VertexCount > 0)
 	{
-		// Calculate accurate bounding box based on raw vertex values
 		JsonAccessor->MinMaxLength = 3;
 
 		for (int32 ComponentIndex = 0; ComponentIndex < JsonAccessor->MinMaxLength; ComponentIndex++)
@@ -176,34 +175,36 @@ FGLTFJsonAccessor* FGLTFPositionBufferConverter::Convert(const FGLTFMeshSection*
 
 FGLTFJsonAccessor* FGLTFPositionDeltaBufferConverter::Convert(const FGLTFMeshSection* MeshSection, const TMap<uint32, FVector3f>* TargetPositionDeltas)
 {
+    if (!TargetPositionDeltas)
+    {
+        return nullptr;
+    }
+
     const TArray<uint32>& IndexMap = MeshSection->IndexMap;
     const uint32 VertexCount = IndexMap.Num() > 0 ? IndexMap.Num() : TargetPositionDeltas->Num();
 
     TArray<FGLTFVector3> Positions;
     Positions.AddUninitialized(VertexCount);
 
-    if (TargetPositionDeltas)
+    if (IndexMap.Num() > 0)
     {
-        if (IndexMap.Num() > 0)
+        for (uint32 VertexIndex = 0; VertexIndex < VertexCount; ++VertexIndex)
         {
-            for (uint32 VertexIndex = 0; VertexIndex < VertexCount; ++VertexIndex)
-            {
-                const uint32 MappedVertexIndex = IndexMap[VertexIndex];
-                const FVector3f* Delta = TargetPositionDeltas->Find(MappedVertexIndex);
-                FVector3f Position = Delta ? *Delta : FVector3f();
+            const uint32 MappedVertexIndex = IndexMap[VertexIndex];
+            const FVector3f* Delta = TargetPositionDeltas->Find(MappedVertexIndex);
+            FVector3f Position = Delta ? *Delta : FVector3f();
 
-                Positions[VertexIndex] = FGLTFCoreUtilities::ConvertPosition(Position, Builder.ExportOptions->ExportUniformScale);
-            }
+            Positions[VertexIndex] = FGLTFCoreUtilities::ConvertPosition(Position, Builder.ExportOptions->ExportUniformScale);
         }
-        else
+    }
+    else
+    {
+        for (uint32 VertexIndex = 0; VertexIndex < VertexCount; ++VertexIndex)
         {
-            for (uint32 VertexIndex = 0; VertexIndex < VertexCount; ++VertexIndex)
-            {
-                const FVector3f* Delta = TargetPositionDeltas->Find(VertexIndex);
-                FVector3f Position = Delta ? *Delta : FVector3f();
+            const FVector3f* Delta = TargetPositionDeltas->Find(VertexIndex);
+            FVector3f Position = Delta ? *Delta : FVector3f();
 
-                Positions[VertexIndex] = FGLTFCoreUtilities::ConvertPosition(Position, Builder.ExportOptions->ExportUniformScale);
-            }
+            Positions[VertexIndex] = FGLTFCoreUtilities::ConvertPosition(Position, Builder.ExportOptions->ExportUniformScale);
         }
     }
 
@@ -373,7 +374,6 @@ FGLTFJsonBufferView* FGLTFNormalBufferConverter::ConvertBufferView(const FGLTFMe
 
 FGLTFJsonAccessor* FGLTFNormalDeltaBufferConverter::Convert(const FGLTFMeshSection* MeshSection, const TMap<uint32, FVector3f>* TargetNormalDeltas, bool bHighPrecision)
 {
-
     FGLTFJsonBufferView* BufferView;
     EGLTFJsonComponentType ComponentType;
 
