[glTF Exporter] Supporting better primitive number calculations, supporting exports for meshes with unused Materials.
[glTF Exporter] Out of bounds crash 'fix' for SourceModel export.

#rnx

#Synced-CL 49989800

diff --git a/Engine/Plugins/Enterprise/GLTFExporter/Source/GLTFExporter/Private/Converters/GLTFMeshConverters.cpp b/Engine/Plugins/Enterprise/GLTFExporter/Source/GLTFExporter/Private/Converters/GLTFMeshConverters.cpp
--- a/Engine/Plugins/Enterprise/GLTFExporter/Source/GLTFExporter/Private/Converters/GLTFMeshConverters.cpp
+++ b/Engine/Plugins/Enterprise/GLTFExporter/Source/GLTFExporter/Private/Converters/GLTFMeshConverters.cpp
@@ -12,6 +12,7 @@
 #include "Components/SkeletalMeshComponent.h"
 #include "Components/SplineMeshComponent.h"
 #include "LandscapeComponent.h"
+#include "Utilities/MeshParsingUtilities.h"
 
 void FGLTFStaticMeshConverter::Sanitize(const UStaticMesh*& StaticMesh, const UStaticMeshComponent*& StaticMeshComponent, FGLTFMaterialArray& Materials, int32& LODIndex)
 {
@@ -39,7 +40,7 @@ void FGLTFStaticMeshConverter::Sanitize(const UStaticMesh*& StaticMesh, const US
 FGLTFJsonMesh* FGLTFStaticMeshConverter::Convert(const UStaticMesh* StaticMesh,  const UStaticMeshComponent* StaticMeshComponent, FGLTFMaterialArray Materials, int32 LODIndex)
 {
 #if !WITH_EDITOR
-	if (!StaticMesh->bAllowCPUAccess)
+	if (StaticMesh && !StaticMesh->bAllowCPUAccess)
 	{
 		Builder.LogSuggestion(FString::Printf(
 			TEXT("Export of mesh %s can in runtime be speed-up by checking 'Allow CPU Access' in asset settings"),
@@ -47,7 +48,21 @@ FGLTFJsonMesh* FGLTFStaticMeshConverter::Convert(const UStaticMesh* StaticMesh,
 	}
 #endif
 
-	const int32 MaterialCount = FGLTFMeshUtilities::GetMaterials(StaticMesh).Num();
+	int32 MaterialCount = 0;
+	if (StaticMesh)
+	{
+#if WITH_EDITORONLY_DATA
+		if (Builder.ExportOptions->bExportSourceModel)
+		{
+			MaterialCount = UE::MeshParser::Private::GetNumberOfValidPolygons(StaticMesh, LODIndex);
+		}
+		else
+#endif
+		{
+			MaterialCount = FGLTFMeshUtilities::GetMaterials(StaticMesh).Num();
+		}
+	}
+
 	FGLTFJsonMesh* JsonMesh = Builder.AddMesh();
 	JsonMesh->Primitives.AddDefaulted(MaterialCount);
 
@@ -81,7 +96,7 @@ void FGLTFSkeletalMeshConverter::Sanitize(const USkeletalMesh*& SkeletalMesh, co
 FGLTFJsonMesh* FGLTFSkeletalMeshConverter::Convert(const USkeletalMesh* SkeletalMesh, const USkeletalMeshComponent* SkeletalMeshComponent, FGLTFMaterialArray Materials, int32 LODIndex)
 {
 #if !WITH_EDITOR
-	if (!SkeletalMesh->GetLODInfo(LODIndex)->bAllowCPUAccess)
+	if (SkeletalMesh && !SkeletalMesh->GetLODInfo(LODIndex)->bAllowCPUAccess)
 	{
 		Builder.LogSuggestion(FString::Printf(
 			TEXT("Export of mesh %s (LOD %d) can in runtime be speed-up by checking 'Allow CPU Access' in asset settings"),
@@ -90,7 +105,21 @@ FGLTFJsonMesh* FGLTFSkeletalMeshConverter::Convert(const USkeletalMesh* Skeletal
 	}
 #endif
 
-	const int32 MaterialCount = FGLTFMeshUtilities::GetMaterials(SkeletalMesh).Num();
+	int32 MaterialCount = 0;
+	if (SkeletalMesh)
+	{
+#if WITH_EDITORONLY_DATA
+		if (Builder.ExportOptions->bExportSourceModel)
+		{
+			MaterialCount = UE::MeshParser::Private::GetNumberOfValidPolygons(SkeletalMesh, LODIndex);
+		}
+		else
+#endif
+		{
+			MaterialCount = FGLTFMeshUtilities::GetMaterials(SkeletalMesh).Num();
+		}
+	}
+
 	FGLTFJsonMesh* JsonMesh = Builder.AddMesh();
 	JsonMesh->Primitives.AddDefaulted(MaterialCount);
 
@@ -120,7 +149,21 @@ FGLTFJsonMesh* FGLTFSplineMeshConverter::Convert(const UStaticMesh* StaticMesh,
 	}
 #endif
 
-	const int32 MaterialCount = StaticMesh ? FGLTFMeshUtilities::GetMaterials(StaticMesh).Num() : 0;
+	int32 MaterialCount = 0;
+	if (StaticMesh)
+	{
+#if WITH_EDITORONLY_DATA
+		if (Builder.ExportOptions->bExportSourceModel)
+		{
+			MaterialCount = UE::MeshParser::Private::GetNumberOfValidPolygons(StaticMesh, LODIndex);
+		}
+		else
+#endif
+		{
+			MaterialCount = FGLTFMeshUtilities::GetMaterials(StaticMesh).Num();
+		}
+	}
+
 	FGLTFJsonMesh* JsonMesh = Builder.AddMesh();
 	JsonMesh->Primitives.AddDefaulted(MaterialCount);
 
diff --git a/Engine/Plugins/Enterprise/GLTFExporter/Source/GLTFExporter/Private/Utilities/MeshParsingUtilities.cpp b/Engine/Plugins/Enterprise/GLTFExporter/Source/GLTFExporter/Private/Utilities/MeshParsingUtilities.cpp
--- a/Engine/Plugins/Enterprise/GLTFExporter/Source/GLTFExporter/Private/Utilities/MeshParsingUtilities.cpp
+++ b/Engine/Plugins/Enterprise/GLTFExporter/Source/GLTFExporter/Private/Utilities/MeshParsingUtilities.cpp
@@ -78,18 +78,6 @@ namespace UE::MeshParser
 		return !(Indices.Num() && Positions.Num());
 	}
 
-	void FMeshPrimitiveDescription::EmptyContainers()
-	{
-		Indices.Empty();
-		Positions.Empty();
-		VertexColors.Empty();
-		Normals.Empty();
-		Tangents.Empty();
-		UVs.Empty();
-		JointInfluences.Empty();
-		JointWeights.Empty();
-	}
-
 	//Except for Joint related containers
 	void FMeshPrimitiveDescription::PrepareContainers(const int32& IndexCount, const int32& AttributesCount, const int32& UVCount, bool bPrepareVertexColors, const int32& TargetCount)
 	{
@@ -103,9 +91,11 @@ namespace UE::MeshParser
 		{
 			UVs[UVIndex].SetNumZeroed(AttributesCount);
 		}
+
+		Targets.SetNum(TargetCount);
 		for (size_t TargetCounter = 0; TargetCounter < TargetCount; TargetCounter++)
 		{
-			FPrimitiveTargetDescription& TargetDescription = Targets.AddDefaulted_GetRef();
+			FPrimitiveTargetDescription& TargetDescription = Targets[TargetCounter];
 			TargetDescription.Positions.SetNumZeroed(AttributesCount);
 			TargetDescription.Normals.SetNumZeroed(AttributesCount);
 		}
@@ -237,20 +227,21 @@ namespace UE::MeshParser
 				continue;
 			}
 
-			MeshPrimitiveDescription.MaterialIndex = GetMaterialIndex(PolygonGroupID);
-
 			//VertexIDs of the TriangleIDs can be part of a bigger/unified container
 			//For exporting purposes we want to have the Primitives to have their own containers for Attributes+Indices.
 			//Which means we have still remap the Vertices to their own containers per Primitives.
 
 			if (!ParseVertexBased(TriangleIDs, MeshPrimitiveDescription, ExportConfigs))
 			{
-				//Emtpy all attribute storages in case the ParseVertexBased already started filling them:
-				MeshPrimitiveDescription.EmptyContainers();
+				//To Emtpy all attribute storages we re-initialize the FMeshPrimitiveDescription.
+				MeshPrimitiveDescription = FMeshPrimitiveDescription();
 
 				ParseVertexInstanceBased(TriangleIDs, MeshPrimitiveDescription, ExportConfigs);
 			}
 
+			//Material Index has to be set here as MeshPrimitiveDescription can be re-initialized if ParseVertexBased fails.
+			MeshPrimitiveDescription.MaterialIndex = GetMaterialIndex(PolygonGroupID);
+
 			if (MeshPrimitiveDescription.Tangents.Num() > 0 && DoesBufferHasZeroVector(MeshPrimitiveDescription.Tangents))
 			{
 				//Do not Export Tangents list that is zeroed out.
@@ -744,4 +735,60 @@ namespace UE::MeshParser
 	template struct FMeshDescriptionParser<FSkeletalMaterial>;
 }
 
+template <class T>
+int32 UE::MeshParser::Private::GetNumberOfValidPolygons(const T* Mesh/*either UStaticMesh or USkeletalMesh*/, const int32 LODIndex)
+{
+	if (!Mesh)
+	{
+		return 0;
+	}
+
+	FMeshDescription* MeshDescription = nullptr;
+	if constexpr (std::is_same_v<T, UStaticMesh>)
+	{
+		MeshDescription = Mesh->GetMeshDescription(LODIndex);
+	}
+	else if constexpr (std::is_same_v<T, USkeletalMesh>)
+	{
+		if (Mesh->GetNumSourceModels() <= LODIndex || !Mesh->IsSourceModelValid(LODIndex))
+		{
+			return 0;
+		}
+		const FSkeletalMeshSourceModel& MeshSourceModel = Mesh->GetSourceModel(LODIndex);
+		MeshDescription = MeshSourceModel.GetMeshDescription();
+	}
+	else
+	{
+		ensure(false);
+		return 0;
+	}
+
+	if (!MeshDescription)
+	{
+		return 0;
+	}
+
+	int PolygonCounter = 0;
+
+	//Iterate through all PolygonsGroups:
+	for (const FPolygonGroupID& PolygonGroupID : MeshDescription->PolygonGroups().GetElementIDs())
+	{
+		//Acquire the TriangleIDs for the current PolygonGroup:
+		const TArrayView<const FTriangleID>& TriangleIDs = MeshDescription->GetPolygonGroupTriangles(PolygonGroupID);
+
+		if (TriangleIDs.Num() == 0)
+		{
+			//Do not export empty primitives.
+			continue;
+		}
+
+		PolygonCounter++;
+	}
+
+	return PolygonCounter;
+}
+
+template int32 UE::MeshParser::Private::GetNumberOfValidPolygons<UStaticMesh>(const UStaticMesh*, const int32 LODIndex);
+template int32 UE::MeshParser::Private::GetNumberOfValidPolygons<USkeletalMesh>(const USkeletalMesh*, const int32 LODIndex);
+
 #endif
\ No newline at end of file
diff --git a/Engine/Plugins/Enterprise/GLTFExporter/Source/GLTFExporter/Public/Utilities/MeshParsingUtilities.h b/Engine/Plugins/Enterprise/GLTFExporter/Source/GLTFExporter/Public/Utilities/MeshParsingUtilities.h
--- a/Engine/Plugins/Enterprise/GLTFExporter/Source/GLTFExporter/Public/Utilities/MeshParsingUtilities.h
+++ b/Engine/Plugins/Enterprise/GLTFExporter/Source/GLTFExporter/Public/Utilities/MeshParsingUtilities.h
@@ -41,7 +41,6 @@ namespace UE::MeshParser
 		TArray<FPrimitiveTargetDescription>		Targets; //MorphTargets
 
 		bool IsEmpty();
-		void EmptyContainers();
 
 		//Except for Joint related containers
 		void PrepareContainers(const int32& IndexCount, const int32& AttributesCount, const int32& UVCount, bool bHasVertexColors, const int32& TargetCount);
@@ -129,11 +128,20 @@ namespace UE::MeshParser
 
 		TConstArrayView<FVertexID> VertexInstanceIdToVertexId;
 		TPolygonGroupAttributesConstRef<FName> PolygonGroupMaterialSlotNames;
-
 	public:
 		//MeshDetails like UVCount/NumberOfPrimitives/bHasVertexColor for ease of access:
 		FMeshDetails MeshDetails;
 	};
+
+	namespace Private
+	{
+		/**
+		* Queries number of valid polygons (polygons that have at least 1 triangle).
+		* It queries it from Source Model (MeshDescription), and not the RenderData.
+		*/
+		template <class T>
+		int32 GetNumberOfValidPolygons(const T* Mesh/*either UStaticMesh or USkeletalMesh*/, const int32 LODIndex);
+	}
 }
 
 #endif
\ No newline at end of file
