[glTF Exporter] Fix for 'Array has changed during ranged-for iteration'.

#rnx

#Synced-CL 49989800

diff --git a/Engine/Plugins/Enterprise/GLTFExporter/Source/GLTFExporter/Private/Builders/GLTFTaskBuilder.cpp b/Engine/Plugins/Enterprise/GLTFExporter/Source/GLTFExporter/Private/Builders/GLTFTaskBuilder.cpp
--- a/Engine/Plugins/Enterprise/GLTFExporter/Source/GLTFExporter/Private/Builders/GLTFTaskBuilder.cpp
+++ b/Engine/Plugins/Enterprise/GLTFExporter/Source/GLTFExporter/Private/Builders/GLTFTaskBuilder.cpp
@@ -47,6 +47,18 @@ void FGLTFTaskBuilder::ProcessSlowTasks(FFeedbackContext* Context)
 				Task->Process();
 			}
 		}
+
+		//Merge over TasksByPrioritySecondary to TasksByPriority
+		//This works here, because 'PriorityIndexLock' only allows addition to higher PriorityIndices.
+		for (TPair<EGLTFTaskPriority, TArray<TUniquePtr<FGLTFDelayedTask>>>& Entry : TasksByPrioritySecondary)
+		{
+			TArray<TUniquePtr<FGLTFDelayedTask>>& TasksToCopyInto = TasksByPriority.FindOrAdd(Entry.Key);
+			for (TUniquePtr<FGLTFDelayedTask>& EntryTask : Entry.Value)
+			{
+				TasksToCopyInto.Add(MoveTemp(EntryTask));
+			}
+		}
+		TasksByPrioritySecondary.Empty();
 	}
 
 	TasksByPriority.Empty();
@@ -68,7 +80,15 @@ bool FGLTFTaskBuilder::ScheduleSlowTask(TUniquePtr<FGLTFDelayedTask> Task)
 		return false;
 	}
 
-	TasksByPriority.FindOrAdd(Priority).Add(MoveTemp(Task));
+	if (PriorityIndexLock == INDEX_NONE)
+	{
+		TasksByPriority.FindOrAdd(Priority).Add(MoveTemp(Task));
+	}
+	else
+	{
+		TasksByPrioritySecondary.FindOrAdd(Priority).Add(MoveTemp(Task));
+	}
+	
 	return true;
 }
 
diff --git a/Engine/Plugins/Enterprise/GLTFExporter/Source/GLTFExporter/Public/Builders/GLTFTaskBuilder.h b/Engine/Plugins/Enterprise/GLTFExporter/Source/GLTFExporter/Public/Builders/GLTFTaskBuilder.h
--- a/Engine/Plugins/Enterprise/GLTFExporter/Source/GLTFExporter/Public/Builders/GLTFTaskBuilder.h
+++ b/Engine/Plugins/Enterprise/GLTFExporter/Source/GLTFExporter/Public/Builders/GLTFTaskBuilder.h
@@ -35,6 +35,7 @@ private:
 
 	int32 PriorityIndexLock;
 	TMap<EGLTFTaskPriority, TArray<TUniquePtr<FGLTFDelayedTask>>> TasksByPriority;
+	TMap<EGLTFTaskPriority, TArray<TUniquePtr<FGLTFDelayedTask>>> TasksByPrioritySecondary; //Adding to this if ProcessSlowTasks is in progress.
 };
 
 #undef UE_API
