Fixed build error of QA_ContentPipeline project introduced by CL 51336460
Fixed glTF exporter level tests failures
- Root cause: Two new (useless) actors have appeared in the resulting glTF file
- Solution: Remove once for all all transient actors from the resulting glTF.

#rnx

#Synced-CL 51338605

diff --git a/Engine/Plugins/Enterprise/GLTFExporter/Source/GLTFExporter/Private/Converters/GLTFNodeConverters.cpp b/Engine/Plugins/Enterprise/GLTFExporter/Source/GLTFExporter/Private/Converters/GLTFNodeConverters.cpp
--- a/Engine/Plugins/Enterprise/GLTFExporter/Source/GLTFExporter/Private/Converters/GLTFNodeConverters.cpp
+++ b/Engine/Plugins/Enterprise/GLTFExporter/Source/GLTFExporter/Private/Converters/GLTFNodeConverters.cpp
@@ -27,7 +27,7 @@
 
 FGLTFJsonNode* FGLTFActorConverter::Convert(const AActor* Actor)
 {
-	if (!Actor || !Actor->IsValidLowLevel() || Actor->IsEditorOnly() || !Builder.IsSelectedActor(Actor))
+	if (!Actor || !Actor->IsValidLowLevel() || Actor->HasAnyFlags(RF_Transient) || Actor->IsEditorOnly() || !Builder.IsSelectedActor(Actor))
 	{
 		return nullptr;
 	}
